Fixed the missing blend shape entry when compiling using older Xcode than the latest greatest

diff --git a/Engine/Plugins/Runtime/AppleARKitFaceSupport/Source/AppleARKitFaceSupport/Private/AppleARKitFaceMeshConversion.h b/Engine/Plugins/Runtime/AppleARKitFaceSupport/Source/AppleARKitFaceSupport/Private/AppleARKitFaceMeshConversion.h
--- a/Engine/Plugins/Runtime/AppleARKitFaceSupport/Source/AppleARKitFaceSupport/Private/AppleARKitFaceMeshConversion.h
+++ b/Engine/Plugins/Runtime/AppleARKitFaceSupport/Source/AppleARKitFaceSupport/Private/AppleARKitFaceMeshConversion.h
@@ -180,6 +180,14 @@ static FORCEINLINE FARBlendShapeMap ToBlendShapeMap(NSDictionary<ARBlendShapeLoc
 	{
 		SET_BLEND_SHAPE(ARBlendShapeLocationTongueOut, EARFaceBlendShape::TongueOut);
 	}
+	else
+	{
+		// Always add a zeroed entry for any unsupported blend shapes
+		BlendShapeMap.Add(EARFaceBlendShape::TongueOut, 0.f);
+	}
+#else
+	// Always add a zeroed entry for any unsupported blend shapes
+	BlendShapeMap.Add(EARFaceBlendShape::TongueOut, 0.f);
 #endif
 
 #undef SET_BLEND_SHAPE
