Fixed a case where a face ar api was still being used in AppleARKit which would cause Apple to reject them

diff --git a/Engine/Plugins/Runtime/AppleARKit/Source/AppleARKit/Private/AppleARKitSystem.cpp b/Engine/Plugins/Runtime/AppleARKit/Source/AppleARKit/Private/AppleARKitSystem.cpp
--- a/Engine/Plugins/Runtime/AppleARKit/Source/AppleARKit/Private/AppleARKitSystem.cpp
+++ b/Engine/Plugins/Runtime/AppleARKit/Source/AppleARKit/Private/AppleARKitSystem.cpp
@@ -739,7 +739,13 @@ bool FAppleARKitSystem::OnIsTrackingTypeSupported(EARSessionType SessionType) co
 		}
 		case EARSessionType::Face:
 		{
-			return ARFaceTrackingConfiguration.isSupported == TRUE;
+			// We need to get the face support from the factory method, which is a modular feature to avoid dependencies
+			TArray<IAppleARKitFaceSupport*> Impls = IModularFeatures::Get().GetModularFeatureImplementations<IAppleARKitFaceSupport>("AppleARKitFaceSupport");
+			if (Impls.Num() > 0 && Impls[0] != nullptr)
+			{
+				return Impls[0]->DoesSupportFaceAR();
+			}
+			return false;
 		}
 	}
 #endif
diff --git a/Engine/Plugins/Runtime/AppleARKit/Source/AppleARKit/Public/AppleARKitFaceSupport.h b/Engine/Plugins/Runtime/AppleARKit/Source/AppleARKit/Public/AppleARKitFaceSupport.h
--- a/Engine/Plugins/Runtime/AppleARKit/Source/AppleARKit/Public/AppleARKitFaceSupport.h
+++ b/Engine/Plugins/Runtime/AppleARKit/Source/AppleARKit/Public/AppleARKitFaceSupport.h
@@ -50,6 +50,11 @@ public:
 	 * @param SessionConfig the UE4 configuration object that needs processing
 	 */
 	virtual ARConfiguration* ToARConfiguration(UARSessionConfig* SessionConfig) { return nullptr; }
+
+	/**
+	 * @return whether this device supports face ar
+	 */
+	virtual bool DoesSupportFaceAR() { return false; }
 #endif
 
 	static FName GetModularFeatureName()
diff --git a/Engine/Plugins/Runtime/AppleARKitFaceSupport/Source/AppleARKitFaceSupport/Private/AppleARKitFaceSupportImpl.cpp b/Engine/Plugins/Runtime/AppleARKitFaceSupport/Source/AppleARKitFaceSupport/Private/AppleARKitFaceSupportImpl.cpp
--- a/Engine/Plugins/Runtime/AppleARKitFaceSupport/Source/AppleARKitFaceSupport/Private/AppleARKitFaceSupportImpl.cpp
+++ b/Engine/Plugins/Runtime/AppleARKitFaceSupport/Source/AppleARKitFaceSupport/Private/AppleARKitFaceSupportImpl.cpp
@@ -135,4 +135,8 @@ void FAppleARKitFaceSupport::PublishLiveLinkData(TSharedPtr<FAppleARKitAnchorDat
 	}
 }
 
+bool FAppleARKitFaceSupport::DoesSupportFaceAR()
+{
+	return ARFaceTrackingConfiguration.isSupported == TRUE;
+}
 #endif
diff --git a/Engine/Plugins/Runtime/AppleARKitFaceSupport/Source/AppleARKitFaceSupport/Private/AppleARKitFaceSupportImpl.h b/Engine/Plugins/Runtime/AppleARKitFaceSupport/Source/AppleARKitFaceSupport/Private/AppleARKitFaceSupportImpl.h
--- a/Engine/Plugins/Runtime/AppleARKitFaceSupport/Source/AppleARKitFaceSupport/Private/AppleARKitFaceSupportImpl.h
+++ b/Engine/Plugins/Runtime/AppleARKitFaceSupport/Source/AppleARKitFaceSupport/Private/AppleARKitFaceSupportImpl.h
@@ -25,6 +25,7 @@ private:
 	virtual TArray<TSharedPtr<FAppleARKitAnchorData>> MakeAnchorData(NSArray<ARAnchor*>* NewAnchors, double Timestamp, uint32 FrameNumber, const FRotator& AdjustBy) override;
 	virtual ARConfiguration* ToARConfiguration(UARSessionConfig* InSessionConfig) override;
 	virtual void PublishLiveLinkData(TSharedPtr<FAppleARKitAnchorData> Anchor, double Timestamp, uint32 FrameNumber) override;
+	virtual bool DoesSupportFaceAR() override;
 	// ~IAppleARKitFaceSupport
 #endif
 
