Fixed Apple automatically mirroring face blendshapes

diff --git a/Engine/Plugins/Runtime/AppleARKitFaceSupport/Source/AppleARKitFaceSupport/Private/AppleARKitFaceMeshComponent.cpp b/Engine/Plugins/Runtime/AppleARKitFaceSupport/Source/AppleARKitFaceSupport/Private/AppleARKitFaceMeshComponent.cpp
--- a/Engine/Plugins/Runtime/AppleARKitFaceSupport/Source/AppleARKitFaceSupport/Private/AppleARKitFaceMeshComponent.cpp
+++ b/Engine/Plugins/Runtime/AppleARKitFaceSupport/Source/AppleARKitFaceSupport/Private/AppleARKitFaceMeshComponent.cpp
@@ -27,8 +27,8 @@ NSDictionary<ARBlendShapeLocation,NSNumber *>* ToBlendShapeDictionary(const FARB
 		BlendShapeDict[AppleShape] = Num; \
 	}
 
-	// Do we want to capture face performance or look at the face as if in a mirror
-	if (GetDefault<UAppleARKitSettings>()->DefaultFaceTrackingDirection == EARFaceTrackingDirection::FaceRelative)
+	// Do we want to capture face performance or look at the face as if in a mirror (Apple is mirrored so we mirror the mirror)
+	if (GetDefault<UAppleARKitSettings>()->DefaultFaceTrackingDirection == EARFaceTrackingDirection::FaceMirrored)
 	{
 		SET_BLEND_SHAPE(ARBlendShapeLocationEyeBlinkLeft, EARFaceBlendShape::EyeBlinkLeft);
 		SET_BLEND_SHAPE(ARBlendShapeLocationEyeLookDownLeft, EARFaceBlendShape::EyeLookDownLeft);
diff --git a/Engine/Plugins/Runtime/AppleARKitFaceSupport/Source/AppleARKitFaceSupport/Private/AppleARKitFaceMeshConversion.h b/Engine/Plugins/Runtime/AppleARKitFaceSupport/Source/AppleARKitFaceSupport/Private/AppleARKitFaceMeshConversion.h
--- a/Engine/Plugins/Runtime/AppleARKitFaceSupport/Source/AppleARKitFaceSupport/Private/AppleARKitFaceMeshConversion.h
+++ b/Engine/Plugins/Runtime/AppleARKitFaceSupport/Source/AppleARKitFaceSupport/Private/AppleARKitFaceMeshConversion.h
@@ -73,8 +73,8 @@ static FORCEINLINE FARBlendShapeMap ToBlendShapeMap(NSDictionary<ARBlendShapeLoc
 		BlendShapeMap.Add(UE4Shape, 0.f); \
 	}
 
-	// Do we want to capture face performance or look at the face as if in a mirror
-	if (GetDefault<UAppleARKitSettings>()->DefaultFaceTrackingDirection == EARFaceTrackingDirection::FaceRelative)
+	// Do we want to capture face performance or look at the face as if in a mirror (Apple is mirrored so we mirror the mirror)
+	if (GetDefault<UAppleARKitSettings>()->DefaultFaceTrackingDirection == EARFaceTrackingDirection::FaceMirrored)
 	{
 		SET_BLEND_SHAPE(ARBlendShapeLocationEyeBlinkLeft, EARFaceBlendShape::EyeBlinkLeft);
 		SET_BLEND_SHAPE(ARBlendShapeLocationEyeLookDownLeft, EARFaceBlendShape::EyeLookDownLeft);
