Merged the workaround for LiveLink not handling ARKit timestamps properly

diff --git a/Engine/Plugins/Runtime/AppleARKitFaceSupport/Source/AppleARKitFaceSupport/Private/AppleARKitLiveLinkSource.cpp b/Engine/Plugins/Runtime/AppleARKitFaceSupport/Source/AppleARKitFaceSupport/Private/AppleARKitLiveLinkSource.cpp
--- a/Engine/Plugins/Runtime/AppleARKitFaceSupport/Source/AppleARKitFaceSupport/Private/AppleARKitLiveLinkSource.cpp
+++ b/Engine/Plugins/Runtime/AppleARKitFaceSupport/Source/AppleARKitFaceSupport/Private/AppleARKitLiveLinkSource.cpp
@@ -132,7 +132,9 @@ void FAppleARKitLiveLinkSource::PublishBlendShapes(FName SubjectName, double Tim
 	{
 		static FLiveLinkFrameData LiveLinkFrame;
 
-		LiveLinkFrame.WorldTime = Timestamp;
+		//LiveLinkFrame.WorldTime = Timestamp;
+		// Don't use the ARKit timestamps because there's a bug somewhere in LiveLink that doesn't interpret them properly
+		LiveLinkFrame.WorldTime = FPlatformTime::Seconds();
 
 		TArray<FLiveLinkCurveElement>& BlendShapes = LiveLinkFrame.CurveElements;
 
