Added mirrored vs face relative support for the blendshapes we create for rotations

diff --git a/Engine/Plugins/Runtime/AppleARKitFaceSupport/Source/AppleARKitFaceSupport/Private/AppleARKitFaceMeshConversion.h b/Engine/Plugins/Runtime/AppleARKitFaceSupport/Source/AppleARKitFaceSupport/Private/AppleARKitFaceMeshConversion.h
--- a/Engine/Plugins/Runtime/AppleARKitFaceSupport/Source/AppleARKitFaceSupport/Private/AppleARKitFaceMeshConversion.h
+++ b/Engine/Plugins/Runtime/AppleARKitFaceSupport/Source/AppleARKitFaceSupport/Private/AppleARKitFaceMeshConversion.h
@@ -37,12 +37,20 @@ static FORCEINLINE TArray<FVector> ToVertexBuffer(const vector_float3* Vertices,
 
 static FORCEINLINE FARBlendShapeMap ToBlendShapeMap(NSDictionary<ARBlendShapeLocation,NSNumber *>* BlendShapes, const FTransform& Transform, const FTransform& LeftEyeTransform, const FTransform& RightEyeTransform)
 {
+	// Do we want to capture face performance or look at the face as if in a mirror (Apple is mirrored so we mirror the mirror)
+	const bool bFaceMirrored = GetDefault<UAppleARKitSettings>()->DefaultFaceTrackingDirection == EARFaceTrackingDirection::FaceMirrored;
+	
 	FARBlendShapeMap BlendShapeMap;
 	FRotator TrackedRot(Transform.GetRotation());
 	// Map the -180..180 range to -1..1
 	float HeadYaw = (float)TrackedRot.Yaw / 180.f;
 	float HeadPitch = (float)TrackedRot.Pitch / 180.f;
 	float HeadRoll = (float)TrackedRot.Roll / 180.f;
+	if (!bFaceMirrored)
+	{
+		HeadYaw = -HeadYaw;
+		HeadRoll = -HeadRoll;
+	}
 	BlendShapeMap.Add(EARFaceBlendShape::HeadYaw, HeadYaw);
 	BlendShapeMap.Add(EARFaceBlendShape::HeadPitch, HeadPitch);
 	BlendShapeMap.Add(EARFaceBlendShape::HeadRoll, HeadRoll);
@@ -51,6 +59,11 @@ static FORCEINLINE FARBlendShapeMap ToBlendShapeMap(NSDictionary<ARBlendShapeLoc
 	float LeftEyeYaw = (float)LeftEyeRot.Yaw / 180.f;
 	float LeftEyePitch = (float)LeftEyeRot.Pitch / 180.f;
 	float LeftEyeRoll = (float)LeftEyeRot.Roll / 180.f;
+	if (!bFaceMirrored)
+	{
+		LeftEyeYaw = -LeftEyeYaw;
+		LeftEyeRoll = -LeftEyeRoll;
+	}
 	BlendShapeMap.Add(EARFaceBlendShape::LeftEyeYaw, LeftEyeYaw);
 	BlendShapeMap.Add(EARFaceBlendShape::LeftEyePitch, LeftEyePitch);
 	BlendShapeMap.Add(EARFaceBlendShape::LeftEyeRoll, LeftEyeRoll);
@@ -58,6 +71,11 @@ static FORCEINLINE FARBlendShapeMap ToBlendShapeMap(NSDictionary<ARBlendShapeLoc
 	float RightEyeYaw = (float)RightEyeRot.Yaw / 180.f;
 	float RightEyePitch = (float)RightEyeRot.Pitch / 180.f;
 	float RightEyeRoll = (float)RightEyeRot.Roll / 180.f;
+	if (!bFaceMirrored)
+	{
+		RightEyeYaw = -RightEyeYaw;
+		RightEyeRoll = -RightEyeRoll;
+	}
 	BlendShapeMap.Add(EARFaceBlendShape::RightEyeYaw, RightEyeYaw);
 	BlendShapeMap.Add(EARFaceBlendShape::RightEyePitch, RightEyePitch);
 	BlendShapeMap.Add(EARFaceBlendShape::RightEyeRoll, RightEyeRoll);
@@ -73,8 +91,7 @@ static FORCEINLINE FARBlendShapeMap ToBlendShapeMap(NSDictionary<ARBlendShapeLoc
 		BlendShapeMap.Add(UE4Shape, 0.f); \
 	}
 
-	// Do we want to capture face performance or look at the face as if in a mirror (Apple is mirrored so we mirror the mirror)
-	if (GetDefault<UAppleARKitSettings>()->DefaultFaceTrackingDirection == EARFaceTrackingDirection::FaceMirrored)
+	if (bFaceMirrored)
 	{
 		SET_BLEND_SHAPE(ARBlendShapeLocationEyeBlinkLeft, EARFaceBlendShape::EyeBlinkLeft);
 		SET_BLEND_SHAPE(ARBlendShapeLocationEyeLookDownLeft, EARFaceBlendShape::EyeLookDownLeft);
