Merging CL#7108152 from Dev-Enterprise

- Updating AppleARKitLiveLinkSource to use new API

diff --git a/Engine/Plugins/Runtime/AR/Apple/AppleARKitFaceSupport/Source/AppleARKitFaceSupport/Private/AppleARKitLiveLinkSource.cpp b/Engine/Plugins/Runtime/AR/Apple/AppleARKitFaceSupport/Source/AppleARKitFaceSupport/Private/AppleARKitLiveLinkSource.cpp
--- a/Engine/Plugins/Runtime/AR/Apple/AppleARKitFaceSupport/Source/AppleARKitFaceSupport/Private/AppleARKitLiveLinkSource.cpp
+++ b/Engine/Plugins/Runtime/AR/Apple/AppleARKitFaceSupport/Source/AppleARKitFaceSupport/Private/AppleARKitLiveLinkSource.cpp
@@ -1,6 +1,7 @@
 // Copyright 1998-2019 Epic Games, Inc. All Rights Reserved.
 
 #include "AppleARKitLiveLinkSource.h"
+
 #include "UObject/Package.h"
 #include "Misc/ConfigCacheIni.h"
 #include "UObject/UObjectGlobals.h"
@@ -20,6 +21,8 @@
 
 #include "AppleARKitFaceSupportModule.h"
 
+#include "Roles/LiveLinkAnimationRole.h"
+
 DECLARE_CYCLE_STAT(TEXT("Publish Local LiveLink"), STAT_FaceAR_Local_PublishLiveLink, STATGROUP_FaceAR);
 DECLARE_CYCLE_STAT(TEXT("Publish Remote LiveLink"), STAT_FaceAR_Remote_PublishLiveLink, STATGROUP_FaceAR);
 DECLARE_CYCLE_STAT(TEXT("Receive LiveLink"), STAT_FaceAR_ReceiveLiveLink, STATGROUP_FaceAR);
@@ -148,64 +151,96 @@ static FName ParseEnumName(FName EnumName)
 	return FName(*EnumString.Right(EnumString.Len() - BlendShapeEnumNameLength));
 }
 
-// Temporary for 4.20
-PRAGMA_DISABLE_DEPRECATION_WARNINGS
-
 void FAppleARKitLiveLinkSource::PublishBlendShapes(FName SubjectName, const FTimecode& Timecode, uint32 FrameRate, const FARBlendShapeMap& FaceBlendShapes, FName DeviceId)
 {
 	SCOPE_CYCLE_COUNTER(STAT_FaceAR_Local_PublishLiveLink);
 
 	check(Client != nullptr);
+	
 	// This code touches UObjects so needs to be run only in the game thread
 	check(IsInGameThread());
 
-	FName* LastSubjectNameForDeviceId = DeviceToLastSubjectNameMap.Find(DeviceId);
+	//If we can't retrieve blend shape enum, nothing we can do
+	const UEnum* EnumPtr = StaticEnum<EARFaceBlendShape>();
+	if (EnumPtr == nullptr)
+	{
+		return;
+	}
+
+	const FLiveLinkSubjectKey SubjectKey(SourceGuid, SubjectName);
+
+
 	// Is this a new device and subject pair?
-	if (LastSubjectNameForDeviceId == nullptr)
+	FBlendShapeStaticData* BlendShapeDataPtr = BlendShapePerDeviceMap.Find(DeviceId);
+	if (BlendShapeDataPtr == nullptr)
 	{
-		// First time seen so publish an empty skeleton
-		Client->PushSubjectSkeleton(SourceGuid, SubjectName, FLiveLinkRefSkeleton());
-		DeviceToLastSubjectNameMap.Add(DeviceId, SubjectName);
+		UpdateStaticData(SubjectName, FaceBlendShapes, DeviceId);
 	}
 	// Did the subject name change for the device?
-	else if (SubjectName != *LastSubjectNameForDeviceId)
+	else if (SubjectKey != BlendShapeDataPtr->SubjectKey)
 	{
 		// The remote device changed subject names, so remove the old subject
-		Client->ClearSubject(*LastSubjectNameForDeviceId);
-		// Now add a new skeleton with the new subject name
-		Client->PushSubjectSkeleton(SourceGuid, SubjectName, FLiveLinkRefSkeleton());
+		Client->RemoveSubject_AnyThread(BlendShapeDataPtr->SubjectKey);
+
+		UpdateStaticData(SubjectName, FaceBlendShapes, DeviceId);
 	}
 
-	const UEnum* EnumPtr = StaticEnum<EARFaceBlendShape>();
-	if (EnumPtr != nullptr)
+	FLiveLinkFrameDataStruct FrameDataStruct(FLiveLinkAnimationFrameData::StaticStruct());
+	FLiveLinkAnimationFrameData* FrameData = FrameDataStruct.Cast<FLiveLinkAnimationFrameData>();
+	FrameData->WorldTime = FPlatformTime::Seconds();
+	FrameData->MetaData.SceneTime = FQualifiedFrameTime(Timecode, FFrameRate(FrameRate, 1));
+	FrameData->PropertyValues.Reserve((int32)EARFaceBlendShape::MAX);
+	
+	// Iterate through all of the blend shapes copying them into the LiveLink data type
+	for (int32 Shape = 0; Shape < (int32)EARFaceBlendShape::MAX; Shape++)
 	{
-		static FLiveLinkFrameData LiveLinkFrame;
+		if (FaceBlendShapes.Contains((EARFaceBlendShape)Shape))
+		{	
+			const float CurveValue = FaceBlendShapes.FindChecked((EARFaceBlendShape)Shape);
+			FrameData->PropertyValues.Add(CurveValue);
+		}
+	}
+
+	//Blendshapes don't change over time. If they were to change, we would need to keep track
+	//of previous values to always have frame data matching static data.
+	check(FrameData->PropertyValues.Num() == BlendShapePerDeviceMap.FindChecked(DeviceId).StaticData.PropertyNames.Num());
+	
+	// Share the data locally with the LiveLink client
+	Client->PushSubjectFrameData_AnyThread(SubjectKey, MoveTemp(FrameDataStruct));
+}
 
-		LiveLinkFrame.WorldTime = FPlatformTime::Seconds();
-		LiveLinkFrame.MetaData.SceneTime = FQualifiedFrameTime(Timecode, FFrameRate(FrameRate, 1));
-		
-		TArray<FLiveLinkCurveElement>& BlendShapes = LiveLinkFrame.CurveElements;
 
-		BlendShapes.Reset((int32)EARFaceBlendShape::MAX);
+void FAppleARKitLiveLinkSource::UpdateStaticData(FName SubjectName, const FARBlendShapeMap& FaceBlendShapes, FName DeviceId /*= NAME_None*/)
+{
+	//Update the subject key to match latest one
+	FBlendShapeStaticData& BlendShapeData = BlendShapePerDeviceMap.FindOrAdd(DeviceId);
+	BlendShapeData.SubjectKey = FLiveLinkSubjectKey(SourceGuid, SubjectName);
 
-		// Iterate through all of the blend shapes copying them into the LiveLink data type
-		for (int32 Shape = 0; Shape < (int32)EARFaceBlendShape::MAX; Shape++)
+	//Update property names array
+	BlendShapeData.StaticData.PropertyNames.Reset((int32)EARFaceBlendShape::MAX);
+
+	//Iterate through all valid blend shapes to extract names
+	const UEnum* EnumPtr = StaticEnum<EARFaceBlendShape>();
+	for (int32 Shape = 0; Shape < (int32)EARFaceBlendShape::MAX; Shape++)
+	{
+		if (FaceBlendShapes.Contains((EARFaceBlendShape)Shape))
 		{
-			if (FaceBlendShapes.Contains((EARFaceBlendShape)Shape))
-			{
-				int32 Index = BlendShapes.AddUninitialized(1);
-				BlendShapes[Index].CurveName = ParseEnumName(EnumPtr->GetNameByValue(Shape));
-				const float CurveValue = FaceBlendShapes.FindChecked((EARFaceBlendShape)Shape);
-				BlendShapes[Index].CurveValue = CurveValue;
-			}
+			//Blendshapes don't change over time. If they were to change, we would need to keep track
+			//of previous values to always have frame data matching static data
+			const FName ShapeName = ParseEnumName(EnumPtr->GetNameByValue(Shape));
+			BlendShapeData.StaticData.PropertyNames.Add(ShapeName);
 		}
-
-		// Share the data locally with the LiveLink client
-		Client->PushSubjectData(SourceGuid, SubjectName, LiveLinkFrame);
 	}
-}
 
-PRAGMA_ENABLE_DEPRECATION_WARNINGS
+	//Push the associated static data
+	FLiveLinkStaticDataStruct StaticDataStruct(FLiveLinkSkeletonStaticData::StaticStruct());
+	FLiveLinkSkeletonStaticData* SkeletonStaticData = StaticDataStruct.Cast<FLiveLinkSkeletonStaticData>();
+	SkeletonStaticData->PropertyNames = BlendShapeData.StaticData.PropertyNames;
+
+	//We're pushing ourselves as Animation Role even if no bones. Property values (curves) are used in anim graph and it requires the animation role
+	UE_LOG(LogAppleARKitFace, Verbose, TEXT("Pushing AppleARKit Subject '%s' with %d blend shapes"), *SubjectName.ToString(), SkeletonStaticData->PropertyNames.Num());
+	Client->PushSubjectStaticData_AnyThread(BlendShapeData.SubjectKey, ULiveLinkAnimationRole::StaticClass(), MoveTemp(StaticDataStruct));
+}
 
 // 1 = Initial version
 // 2 = ARKit 2.0 extra blendshapes
diff --git a/Engine/Plugins/Runtime/AR/Apple/AppleARKitFaceSupport/Source/AppleARKitFaceSupport/Private/AppleARKitLiveLinkSource.h b/Engine/Plugins/Runtime/AR/Apple/AppleARKitFaceSupport/Source/AppleARKitFaceSupport/Private/AppleARKitLiveLinkSource.h
--- a/Engine/Plugins/Runtime/AR/Apple/AppleARKitFaceSupport/Source/AppleARKitFaceSupport/Private/AppleARKitLiveLinkSource.h
+++ b/Engine/Plugins/Runtime/AR/Apple/AppleARKitFaceSupport/Source/AppleARKitFaceSupport/Private/AppleARKitLiveLinkSource.h
@@ -5,6 +5,8 @@
 #include "CoreMinimal.h"
 #include "AppleARKitLiveLinkSourceFactory.h"
 #include "ILiveLinkClient.h"
+#include "LiveLinkTypes.h"
+#include "Roles/LiveLinkAnimationTypes.h"
 #include "Tickable.h"
 #include "ARTrackable.h"
 
@@ -889,6 +891,11 @@ private:
 	virtual void PublishBlendShapes(FName SubjectName, const FTimecode& Timecode, uint32 FrameRate, const FARBlendShapeMap& FaceBlendShapes, FName DeviceId = NAME_None) override;
 	// End IARKitBlendShapePublisher
 
+	// Update static data property names from blend shape map. Should only happen once per subject name. Pushes static data to livelink to register subject
+	void UpdateStaticData(FName SubjectName, const FARBlendShapeMap& FaceBlendShapes, FName DeviceId);
+
+private:
+
 	/** The local client to push data updates to */
 	ILiveLinkClient* Client;
 
@@ -898,6 +905,13 @@ private:
 	/** The last time we sent the data. Used to not send redundant data */
 	uint32 LastFramePublished;
 
-	/** Used to track names changes for a given device */
-	TMap<FName, FName> DeviceToLastSubjectNameMap;
+	/** Helper struct to contain BlendShape data per deviceId */
+	struct FBlendShapeStaticData
+	{
+		FLiveLinkSubjectKey SubjectKey;
+		FLiveLinkSkeletonStaticData StaticData;
+	};
+
+	/** Used to track names changes for a given device and keep track of property names contained in subject */
+	TMap<FName, FBlendShapeStaticData> BlendShapePerDeviceMap;
 };
diff --git a/Engine/Source/Runtime/LiveLinkInterface/Public/LiveLinkTypes.h b/Engine/Source/Runtime/LiveLinkInterface/Public/LiveLinkTypes.h
--- a/Engine/Source/Runtime/LiveLinkInterface/Public/LiveLinkTypes.h
+++ b/Engine/Source/Runtime/LiveLinkInterface/Public/LiveLinkTypes.h
@@ -68,6 +68,7 @@ struct LIVELINKINTERFACE_API FLiveLinkSubjectKey
 	FLiveLinkSubjectKey(const FLiveLinkSubjectKey& Rhs) : Source(Rhs.Source), SubjectName(Rhs.SubjectName) {}
 
 	bool operator== (const FLiveLinkSubjectKey& Other) const { return SubjectName == Other.SubjectName && Source == Other.Source; }
+	bool operator!=(const FLiveLinkSubjectKey& Other) const	{ return !(*this == Other); }
 
 	friend FArchive& operator<<(FArchive& Ar, FLiveLinkSubjectKey& InSubjectKey)
 	{
