- Changing AppleARKitLiveLinkSource to push itself as a Basicrole now that it's supported in anim graphs

diff --git a/Engine/Plugins/Runtime/AR/Apple/AppleARKitFaceSupport/Source/AppleARKitFaceSupport/Private/AppleARKitLiveLinkSource.cpp b/Engine/Plugins/Runtime/AR/Apple/AppleARKitFaceSupport/Source/AppleARKitFaceSupport/Private/AppleARKitLiveLinkSource.cpp
--- a/Engine/Plugins/Runtime/AR/Apple/AppleARKitFaceSupport/Source/AppleARKitFaceSupport/Private/AppleARKitLiveLinkSource.cpp
+++ b/Engine/Plugins/Runtime/AR/Apple/AppleARKitFaceSupport/Source/AppleARKitFaceSupport/Private/AppleARKitLiveLinkSource.cpp
@@ -21,7 +21,7 @@
 
 #include "AppleARKitFaceSupportModule.h"
 
-#include "Roles/LiveLinkAnimationRole.h"
+#include "Roles/LiveLinkBasicRole.h"
 
 DECLARE_CYCLE_STAT(TEXT("Publish Local LiveLink"), STAT_FaceAR_Local_PublishLiveLink, STATGROUP_FaceAR);
 DECLARE_CYCLE_STAT(TEXT("Publish Remote LiveLink"), STAT_FaceAR_Remote_PublishLiveLink, STATGROUP_FaceAR);
@@ -185,8 +185,8 @@ void FAppleARKitLiveLinkSource::PublishBlendShapes(FName SubjectName, const FTim
 		UpdateStaticData(SubjectName, FaceBlendShapes, DeviceId);
 	}
 
-	FLiveLinkFrameDataStruct FrameDataStruct(FLiveLinkAnimationFrameData::StaticStruct());
-	FLiveLinkAnimationFrameData* FrameData = FrameDataStruct.Cast<FLiveLinkAnimationFrameData>();
+	FLiveLinkFrameDataStruct FrameDataStruct(FLiveLinkBaseFrameData::StaticStruct());
+	FLiveLinkBaseFrameData* FrameData = FrameDataStruct.Cast<FLiveLinkBaseFrameData>();
 	FrameData->WorldTime = FPlatformTime::Seconds();
 	FrameData->MetaData.SceneTime = FQualifiedFrameTime(Timecode, FFrameRate(FrameRate, 1));
 	FrameData->PropertyValues.Reserve((int32)EARFaceBlendShape::MAX);
@@ -233,13 +233,12 @@ void FAppleARKitLiveLinkSource::UpdateStaticData(FName SubjectName, const FARBle
 	}
 
 	//Push the associated static data
-	FLiveLinkStaticDataStruct StaticDataStruct(FLiveLinkSkeletonStaticData::StaticStruct());
-	FLiveLinkSkeletonStaticData* SkeletonStaticData = StaticDataStruct.Cast<FLiveLinkSkeletonStaticData>();
-	SkeletonStaticData->PropertyNames = BlendShapeData.StaticData.PropertyNames;
+	FLiveLinkStaticDataStruct StaticDataStruct(FLiveLinkBaseStaticData::StaticStruct());
+	FLiveLinkBaseStaticData* BaseStaticData = StaticDataStruct.Cast<FLiveLinkBaseStaticData>();
+	BaseStaticData->PropertyNames = BlendShapeData.StaticData.PropertyNames;
 
-	//We're pushing ourselves as Animation Role even if no bones. Property values (curves) are used in anim graph and it requires the animation role
-	UE_LOG(LogAppleARKitFace, Verbose, TEXT("Pushing AppleARKit Subject '%s' with %d blend shapes"), *SubjectName.ToString(), SkeletonStaticData->PropertyNames.Num());
-	Client->PushSubjectStaticData_AnyThread(BlendShapeData.SubjectKey, ULiveLinkAnimationRole::StaticClass(), MoveTemp(StaticDataStruct));
+	UE_LOG(LogAppleARKitFace, Verbose, TEXT("Pushing AppleARKit Subject '%s' with %d blend shapes"), *SubjectName.ToString(), BaseStaticData->PropertyNames.Num());
+	Client->PushSubjectStaticData_AnyThread(BlendShapeData.SubjectKey, ULiveLinkBasicRole::StaticClass(), MoveTemp(StaticDataStruct));
 }
 
 // 1 = Initial version
