Fixed export issue that didn't show up on Mac

diff --git a/Engine/Plugins/Runtime/AppleARKitFaceSupport/Source/AppleARKitFaceSupport/Public/AppleARKitFaceMeshComponent.h b/Engine/Plugins/Runtime/AppleARKitFaceSupport/Source/AppleARKitFaceSupport/Public/AppleARKitFaceMeshComponent.h
--- a/Engine/Plugins/Runtime/AppleARKitFaceSupport/Source/AppleARKitFaceSupport/Public/AppleARKitFaceMeshComponent.h
+++ b/Engine/Plugins/Runtime/AppleARKitFaceSupport/Source/AppleARKitFaceSupport/Public/AppleARKitFaceMeshComponent.h
@@ -27,7 +27,7 @@ enum class EARFaceComponentTransformMixing : uint8
  * This component is updated by the ARSystem with face data on devices that have support for it
  */
 UCLASS(hidecategories = (Object, LOD, "Components|ProceduralMesh"), meta = (BlueprintSpawnableComponent), ClassGroup = "AR")
-class APPLEARKIT_API UAppleARKitFaceMeshComponent :
+class APPLEARKITFACESUPPORT_API UAppleARKitFaceMeshComponent :
 	public UProceduralMeshComponent
 {
 	GENERATED_UCLASS_BODY()
diff --git a/Engine/Plugins/Runtime/AppleARKitFaceSupport/Source/AppleARKitFaceSupport/Public/AppleARKitLiveLinkSourceFactory.h b/Engine/Plugins/Runtime/AppleARKitFaceSupport/Source/AppleARKitFaceSupport/Public/AppleARKitLiveLinkSourceFactory.h
--- a/Engine/Plugins/Runtime/AppleARKitFaceSupport/Source/AppleARKitFaceSupport/Public/AppleARKitLiveLinkSourceFactory.h
+++ b/Engine/Plugins/Runtime/AppleARKitFaceSupport/Source/AppleARKitFaceSupport/Public/AppleARKitLiveLinkSourceFactory.h
@@ -8,21 +8,21 @@
 
 
 /** Interface that publishes face ar blend shape information */
-class APPLEARKIT_API IARKitBlendShapePublisher
+class APPLEARKITFACESUPPORT_API IARKitBlendShapePublisher
 {
 public:
 	virtual void PublishBlendShapes(FName SubjectName, double Timestamp, uint32 FrameNumber, const FARBlendShapeMap& BlendShapes) = 0;
 };
 
 /** Interface that publishes face ar blend shape information via LiveLink */
-class APPLEARKIT_API ILiveLinkSourceARKit :
+class APPLEARKITFACESUPPORT_API ILiveLinkSourceARKit :
 	public IARKitBlendShapePublisher,
 	public ILiveLinkSource
 {
 };
 
 /** Factory that creates and registers the sources with the LiveLink client */
-class APPLEARKIT_API FAppleARKitLiveLinkSourceFactory
+class APPLEARKITFACESUPPORT_API FAppleARKitLiveLinkSourceFactory
 {
 public:
 	/** Creates a face mesh source that will autobind to the tracked face mesh */
