Fixed missing include that didn't show up with non-unity compiling on Mac for some reason

diff --git a/Engine/Plugins/Runtime/AppleARKitFaceSupport/Source/AppleARKitFaceSupport/Private/AppleARKitFaceSupportImpl.cpp b/Engine/Plugins/Runtime/AppleARKitFaceSupport/Source/AppleARKitFaceSupport/Private/AppleARKitFaceSupportImpl.cpp
--- a/Engine/Plugins/Runtime/AppleARKitFaceSupport/Source/AppleARKitFaceSupport/Private/AppleARKitFaceSupportImpl.cpp
+++ b/Engine/Plugins/Runtime/AppleARKitFaceSupport/Source/AppleARKitFaceSupport/Private/AppleARKitFaceSupportImpl.cpp
@@ -6,6 +6,7 @@
 #include "AppleARKitConfiguration.h"
 #include "Async/TaskGraphInterfaces.h"
 #include "ARSystem.h"
+#include "Misc/ConfigCacheIni.h"
 
 DECLARE_STATS_GROUP(TEXT("AppleARKitFaceSupport"), STATGROUP_APPLEARKITFACE, STATCAT_Advanced);
 
