Fixed missing include and missing export

diff --git a/Engine/Plugins/Runtime/AppleARKitFaceSupport/Source/AppleARKitFaceSupport/Private/AppleARKitFaceSupportImpl.h b/Engine/Plugins/Runtime/AppleARKitFaceSupport/Source/AppleARKitFaceSupport/Private/AppleARKitFaceSupportImpl.h
--- a/Engine/Plugins/Runtime/AppleARKitFaceSupport/Source/AppleARKitFaceSupport/Private/AppleARKitFaceSupportImpl.h
+++ b/Engine/Plugins/Runtime/AppleARKitFaceSupport/Source/AppleARKitFaceSupport/Private/AppleARKitFaceSupportImpl.h
@@ -24,7 +24,7 @@ struct FAppleARKitFaceAnchorData :
 	static TArray<int32> FaceIndices;
 };
 
-class FAppleARKitFaceSupport :
+class APPLEARKITFACESUPPORT_API FAppleARKitFaceSupport :
 	public FAppleARKitFaceSupportBase
 {
 public:
diff --git a/Engine/Plugins/Runtime/AppleARKitFaceSupport/Source/AppleARKitFaceSupport/Private/AppleARKitFaceSupportModule.cpp b/Engine/Plugins/Runtime/AppleARKitFaceSupport/Source/AppleARKitFaceSupport/Private/AppleARKitFaceSupportModule.cpp
--- a/Engine/Plugins/Runtime/AppleARKitFaceSupport/Source/AppleARKitFaceSupport/Private/AppleARKitFaceSupportModule.cpp
+++ b/Engine/Plugins/Runtime/AppleARKitFaceSupport/Source/AppleARKitFaceSupport/Private/AppleARKitFaceSupportModule.cpp
@@ -6,13 +6,14 @@
 #include "Features/IModularFeatures.h"
 #include "CoreGlobals.h"
 #include "AppleARKitFaceSupportImpl.h"
+#include "Misc/ConfigCacheIni.h"
 
 
 IMPLEMENT_MODULE(FAppleARKitFaceSupportModule, AppleARKitFaceSupport);
 
 DEFINE_LOG_CATEGORY(LogAppleARKitFace);
 
-class FAppleARKitFaceSupportFactory :
+class APPLEARKITFACESUPPORT_API FAppleARKitFaceSupportFactory :
 	public IAppleARKitFaceSupportFactory
 {
 public:
