Changed where face LiveLink was logging to

diff --git a/Engine/Plugins/Runtime/AppleARKitFaceSupport/Source/AppleARKitFaceSupport/Private/AppleARKitLiveLinkSource.cpp b/Engine/Plugins/Runtime/AppleARKitFaceSupport/Source/AppleARKitFaceSupport/Private/AppleARKitLiveLinkSource.cpp
--- a/Engine/Plugins/Runtime/AppleARKitFaceSupport/Source/AppleARKitFaceSupport/Private/AppleARKitLiveLinkSource.cpp
+++ b/Engine/Plugins/Runtime/AppleARKitFaceSupport/Source/AppleARKitFaceSupport/Private/AppleARKitLiveLinkSource.cpp
@@ -17,6 +17,8 @@
 #include "AppleARKitSettings.h"
 #include "ARTrackable.h"
 
+#include "AppleARKitFaceSupportModule.h"
+
 
 TSharedPtr<ILiveLinkSourceARKit> FAppleARKitLiveLinkSourceFactory::CreateLiveLinkSource(bool bCreateRemotePublisher)
 {
@@ -62,7 +64,7 @@ FAppleARKitLiveLinkSource::FAppleARKitLiveLinkSource(bool bCreateRemotePublisher
 		}
 		else
 		{
-			UE_LOG(LogAppleARKit, Warning, TEXT("Failed to create LiveLink remote publisher, so no data will be sent out"));
+			UE_LOG(LogAppleARKitFace, Warning, TEXT("Failed to create LiveLink remote publisher, so no data will be sent out"));
 			delete Publisher;
 		}
 	}
@@ -177,7 +179,7 @@ bool FAppleARKitLiveLinkSource::Exec(UWorld* InWorld, const TCHAR* Cmd, FOutputD
 			}
 			else
 			{
-				UE_LOG(LogAppleARKit, Warning, TEXT("Failed to create LiveLink remote publisher, so no data will be sent out"));
+				UE_LOG(LogAppleARKitFace, Warning, TEXT("Failed to create LiveLink remote publisher, so no data will be sent out"));
 				delete Publisher;
 			}
 #endif
@@ -245,7 +247,7 @@ TSharedRef<FInternetAddr> FAppleARKitLiveLinkRemotePublisher::GetSendAddress()
 		SendAddr->SetPort(LiveLinkPort);
 		bool bIsValid = false;
 		SendAddr->SetIp(*RemoteIp, bIsValid);
-		UE_LOG(LogAppleARKit, Log, TEXT("Sending LiveLink face AR data to address (%s)"), *SendAddr->ToString(true));
+		UE_LOG(LogAppleARKitFace, Log, TEXT("Sending LiveLink face AR data to address (%s)"), *SendAddr->ToString(true));
 	}
 	return SendAddr;
 }
@@ -277,7 +279,7 @@ void FAppleARKitLiveLinkRemotePublisher::PublishBlendShapes(FName SubjectName, d
 			AmountSent != SourceBufferSize)
 		{
 			ISocketSubsystem* SocketSub = ISocketSubsystem::Get();
-			UE_LOG(LogAppleARKit, Verbose, TEXT("Failed to send face AR packet with error (%s). Packet size (%d), sent (%d)"), SocketSub->GetSocketError(), SourceBufferSize, AmountSent);
+			UE_LOG(LogAppleARKitFace, Verbose, TEXT("Failed to send face AR packet with error (%s). Packet size (%d), sent (%d)"), SocketSub->GetSocketError(), SourceBufferSize, AmountSent);
 		}
 	}
 }
@@ -318,7 +320,7 @@ bool FAppleARKitLiveLinkRemoteListener::InitReceiveSocket()
 		{
 			SocketSubsystem->DestroySocket(RecvSocket);
 			RecvSocket = nullptr;
-			UE_LOG(LogAppleARKit, Warning, TEXT("Failed to bind to the listen port (%s) for LiveLink face AR receiving with error (%s)"),
+			UE_LOG(LogAppleARKitFace, Warning, TEXT("Failed to bind to the listen port (%s) for LiveLink face AR receiving with error (%s)"),
 				*Addr->ToString(true), SocketSubsystem->GetSocketError());
 		}
 	}
@@ -357,7 +359,7 @@ void FAppleARKitLiveLinkRemoteListener::Tick(float DeltaTime)
 			FromBuffer >> PacketVer;
 			if (FromBuffer.HasOverflow() || PacketVer != BLEND_SHAPE_PACKET_VER)
 			{
-				UE_LOG(LogAppleARKit, Verbose, TEXT("Packet overflow reading the packet version for the face AR packet"));
+				UE_LOG(LogAppleARKitFace, Verbose, TEXT("Packet overflow reading the packet version for the face AR packet"));
 				return;
 			}
 			FromBuffer >> SubjectName;
@@ -366,7 +368,7 @@ void FAppleARKitLiveLinkRemoteListener::Tick(float DeltaTime)
 			FromBuffer >> BlendShapeCount;
 			if (FromBuffer.HasOverflow() || BlendShapeCount != (uint8)EARFaceBlendShape::MAX)
 			{
-				UE_LOG(LogAppleARKit, Verbose, TEXT("Packet overflow reading the face AR packet's non-array fields"));
+				UE_LOG(LogAppleARKitFace, Verbose, TEXT("Packet overflow reading the face AR packet's non-array fields"));
 				return;
 			}
 
@@ -388,7 +390,7 @@ void FAppleARKitLiveLinkRemoteListener::Tick(float DeltaTime)
 			}
 			else
 			{
-				UE_LOG(LogAppleARKit, Verbose, TEXT("Packet overflow reading the face AR packet's array of blend shapes"));
+				UE_LOG(LogAppleARKitFace, Verbose, TEXT("Packet overflow reading the face AR packet's array of blend shapes"));
 			}
 		}
 	}
