Dramatically simplified the face API isolation to a separate plugin

diff --git a/Engine/Plugins/Runtime/AppleARKit/Source/AppleARKit/Private/AppleARKitSystem.cpp b/Engine/Plugins/Runtime/AppleARKit/Source/AppleARKit/Private/AppleARKitSystem.cpp
--- a/Engine/Plugins/Runtime/AppleARKit/Source/AppleARKit/Private/AppleARKitSystem.cpp
+++ b/Engine/Plugins/Runtime/AppleARKit/Source/AppleARKit/Private/AppleARKitSystem.cpp
@@ -157,6 +157,7 @@ FAppleARKitSystem::FAppleARKitSystem()
 , GameThreadFrameNumber(0)
 , GameThreadTimestamp(0.0)
 , LastTrackedGeometry_DebugId(0)
+, FaceARSupport(nullptr)
 {
 	// See Initialize(), as we need access to SharedThis()
 }
@@ -171,7 +172,7 @@ void FAppleARKitSystem::Shutdown()
 #if SUPPORTS_ARKIT_1_0
 	if (Session != nullptr)
 	{
-		FaceARSupport = TSharedPtr<FAppleARKitFaceSupportBase, ESPMode::ThreadSafe>();
+		FaceARSupport = nullptr;
 		[Session pause];
 		Session.delegate = nullptr;
 		[Session release];
@@ -187,7 +188,7 @@ void FAppleARKitSystem::CheckForFaceARSupport(UARSessionConfig* InSessionConfig)
 	if (InSessionConfig->GetSessionType() != EARSessionType::Face)
 	{
 		// Clear the face ar support so we don't forward to it
-		FaceARSupport = TSharedPtr<FAppleARKitFaceSupportBase, ESPMode::ThreadSafe>();
+		FaceARSupport = nullptr;
 		return;
 	}
 
@@ -196,23 +197,11 @@ void FAppleARKitSystem::CheckForFaceARSupport(UARSessionConfig* InSessionConfig)
 	if (ensureAlwaysMsgf(Feature != nullptr, TEXT("Face AR session has been requested but the face ar plugin is not enabled")))
 	{
 		IAppleARKitFaceSupportFactory* Factory = static_cast<IAppleARKitFaceSupportFactory*>(Feature);
-		FaceARSupport = Factory->CreateFaceSupport(SharedThis(this), this);
-		ensureAlwaysMsgf(FaceARSupport.IsValid(), TEXT("Face AR session has been requested but the face ar plugin is not enabled"));
+		FaceARSupport = Factory->CreateFaceSupport();
+		ensureAlwaysMsgf(FaceARSupport != nullptr, TEXT("Face AR session has been requested but the face ar plugin is not enabled"));
 	}
 }
 
-UARTrackedGeometry* FAppleARKitSystem::GetTrackedGeometry(const FGuid& GeoGuid)
-{
-	UARTrackedGeometry** GeometrySearchResult = TrackedGeometries.Find(GeoGuid);
-	return GeometrySearchResult != nullptr ? *GeometrySearchResult : nullptr;
-}
-
-void FAppleARKitSystem::AddTrackedGeometry(const FGuid& Guid, UARTrackedGeometry* TrackedGeo)
-{
-	check(nullptr);
-	TrackedGeometries.Add(Guid, TrackedGeo);
-}
-
 FName FAppleARKitSystem::GetSystemName() const
 {
 	static const FName AppleARKitSystemName(TEXT("AppleARKit"));
@@ -848,7 +837,7 @@ bool FAppleARKitSystem::Run(UARSessionConfig* SessionConfig)
 
 		ARConfiguration* Configuration = nullptr;
 		CheckForFaceARSupport(SessionConfig);
-		if (!FaceARSupport.IsValid())
+		if (FaceARSupport == nullptr)
 		{
 			Configuration = ToARConfiguration(SessionConfig, Config, CandidateImages, ConvertedCandidateImages);
 		}
@@ -979,6 +968,8 @@ void FAppleARKitSystem::SessionDidFailWithError_DelegateThread(const FString& Er
 
 #if SUPPORTS_ARKIT_1_0
 
+TArray<int32> FAppleARKitAnchorData::FaceIndices;
+
 static TSharedPtr<FAppleARKitAnchorData> MakeAnchorData( ARAnchor* Anchor )
 {
 	TSharedPtr<FAppleARKitAnchorData> NewAnchor;
@@ -1035,9 +1026,14 @@ void FAppleARKitSystem::SessionDidAddAnchors_DelegateThread( NSArray<ARAnchor*>*
 					   STATGROUP_APPLEARKIT);
 
 	// If this object is valid, we are running a face session and need that code to process things
-	if (FaceARSupport.IsValid())
+	if (FaceARSupport != nullptr)
 	{
-		FaceARSupport->ProcessAnchorAdd(anchors, GetAlignmentTransform(), GameThreadFrameNumber, GameThreadTimestamp);
+		const TArray<TSharedPtr<FAppleARKitAnchorData>> AnchorList = FaceARSupport->MakeAnchorData(anchors, GameThreadTimestamp, GameThreadFrameNumber);
+		for (TSharedPtr<FAppleARKitAnchorData> NewAnchorData : AnchorList)
+		{
+			auto AddAnchorTask = FSimpleDelegateGraphTask::FDelegate::CreateSP(this, &FAppleARKitSystem::SessionDidAddAnchors_Internal, NewAnchorData.ToSharedRef());
+			FSimpleDelegateGraphTask::CreateAndDispatchWhenReady(AddAnchorTask, GET_STATID(STAT_FAppleARKitSystem_SessionDidAddAnchors), nullptr, ENamedThreads::GameThread);
+		}
 		return;
 	}
 
@@ -1059,9 +1055,14 @@ void FAppleARKitSystem::SessionDidUpdateAnchors_DelegateThread( NSArray<ARAnchor
 					   STATGROUP_APPLEARKIT);
 	
 	// If this object is valid, we are running a face session and need that code to process things
-	if (FaceARSupport.IsValid())
+	if (FaceARSupport != nullptr)
 	{
-		FaceARSupport->ProcessAnchorUpdate(anchors, GetAlignmentTransform(), GameThreadFrameNumber, GameThreadTimestamp);
+		const TArray<TSharedPtr<FAppleARKitAnchorData>> AnchorList = FaceARSupport->MakeAnchorData(anchors, GameThreadTimestamp, GameThreadFrameNumber);
+		for (TSharedPtr<FAppleARKitAnchorData> NewAnchorData : AnchorList)
+		{
+			auto UpdateAnchorTask = FSimpleDelegateGraphTask::FDelegate::CreateSP(this, &FAppleARKitSystem::SessionDidUpdateAnchors_Internal, NewAnchorData.ToSharedRef());
+			FSimpleDelegateGraphTask::CreateAndDispatchWhenReady(UpdateAnchorTask, GET_STATID(STAT_FAppleARKitSystem_SessionDidUpdateAnchors), nullptr, ENamedThreads::GameThread);
+		}
 		return;
 	}
 
@@ -1132,6 +1133,14 @@ void FAppleARKitSystem::SessionDidAddAnchors_Internal( TSharedRef<FAppleARKitAnc
 			NewGeometry = NewImage;
 			break;
 		}
+		case EAppleAnchorType::FaceAnchor:
+		{
+			UARFaceGeometry* NewGeo = NewObject<UARFaceGeometry>();
+			NewAnchorDebugName = FString::Printf(TEXT("FACE-%02d"), LastTrackedGeometry_DebugId++);
+			NewGeo->SetDebugName(FName(*NewAnchorDebugName));
+			NewGeo->UpdateTrackedGeometry(SharedThis(this), GameThreadFrameNumber, GameThreadTimestamp, AnchorData->Transform, GetAlignmentTransform(), AnchorData->BlendShapes, AnchorData->FaceVerts, AnchorData->FaceIndices);
+			NewGeometry = NewGeo;
+		}
 		default:
 		{
 			NewAnchorDebugName = FString::Printf(TEXT("APPL-%02d"), LastTrackedGeometry_DebugId++);
@@ -1193,6 +1202,19 @@ void FAppleARKitSystem::SessionDidUpdateAnchors_Internal( TSharedRef<FAppleARKit
 				}
 				break;
 			}
+			case EAppleAnchorType::FaceAnchor:
+			{
+				if (UARFaceGeometry* FaceGeo = Cast<UARFaceGeometry>(FoundGeometry))
+				{
+					FaceGeo->UpdateTrackedGeometry(SharedThis(this), GameThreadFrameNumber, GameThreadTimestamp, AnchorData->Transform, GetAlignmentTransform(), AnchorData->BlendShapes, AnchorData->FaceVerts, AnchorData->FaceIndices);
+					for (UARPin* Pin : PinsToUpdate)
+					{
+						const FTransform Pin_LocalToTrackingTransform_PostUpdate = Pin->GetLocalToTrackingTransform_NoAlignment() * AnchorDeltaTransform;
+						Pin->OnTransformUpdated(Pin_LocalToTrackingTransform_PostUpdate);
+					}
+				}
+				break;
+			}
 		}
 	}
 }
diff --git a/Engine/Plugins/Runtime/AppleARKit/Source/AppleARKit/Private/AppleARKitSystem.h b/Engine/Plugins/Runtime/AppleARKit/Source/AppleARKit/Private/AppleARKitSystem.h
--- a/Engine/Plugins/Runtime/AppleARKit/Source/AppleARKit/Private/AppleARKitSystem.h
+++ b/Engine/Plugins/Runtime/AppleARKit/Source/AppleARKit/Private/AppleARKitSystem.h
@@ -27,8 +27,7 @@ struct FAppleARKitFrame;
 struct FAppleARKitAnchorData;
 
 class FAppleARKitSystem :
-	public FARSystemBase,
-	public IAppleARKitFaceSupportCallback
+	public FARSystemBase
 {
 	friend class FAppleARKitXRCamera;
 	
@@ -50,11 +49,6 @@ public:
 	bool OnStartGameFrame(FWorldContext& WorldContext) override;
 	//~ IXRTrackingSystem
 
-	//~IAppleARKitFaceSupportCallback
-	virtual UARTrackedGeometry* GetTrackedGeometry(const FGuid& GeoGuid) override;
-	virtual void AddTrackedGeometry(const FGuid& Guid, UARTrackedGeometry* TrackedGeo) override;
-	//~IAppleARKitFaceSupportCallback
-
 	void* GetARSessionRawPointer() override;
 	void* GetGameThreadARFrameRawPointer() override;
 	
@@ -195,7 +189,7 @@ private:
 	TSharedPtr< FAppleARKitFrame, ESPMode::ThreadSafe > LastReceivedFrame;
 
 	// The object that is handling face support if present
-	TSharedPtr<FAppleARKitFaceSupportBase, ESPMode::ThreadSafe> FaceARSupport;
+	IAppleARKitFaceSupport* FaceARSupport;
 };
 
 
diff --git a/Engine/Plugins/Runtime/AppleARKit/Source/AppleARKit/Public/AppleARKitConversion.h b/Engine/Plugins/Runtime/AppleARKit/Source/AppleARKit/Public/AppleARKitConversion.h
--- a/Engine/Plugins/Runtime/AppleARKit/Source/AppleARKit/Public/AppleARKitConversion.h
+++ b/Engine/Plugins/Runtime/AppleARKit/Source/AppleARKit/Public/AppleARKitConversion.h
@@ -145,6 +145,15 @@ struct FAppleARKitAnchorData
 	{
 	}
 
+	FAppleARKitAnchorData(FGuid InAnchorGuid, FTransform InTransform, FARBlendShapeMap InBlendShapes, TArray<FVector> InFaceVerts)
+		: Transform( InTransform )
+		, AnchorType( EAppleAnchorType::FaceAnchor )
+		, AnchorGUID( InAnchorGuid )
+		, BlendShapes( MoveTemp(InBlendShapes) )
+		, FaceVerts( MoveTemp(InFaceVerts) )
+	{
+	}
+
 	FTransform Transform;
 	EAppleAnchorType AnchorType;
 	FGuid AnchorGUID;
@@ -153,6 +162,11 @@ struct FAppleARKitAnchorData
 	TArray<FVector> BoundaryVerts;
 
 	FString DetectedImageName;
+
+	FARBlendShapeMap BlendShapes;
+	TArray<FVector> FaceVerts;
+	// Note: the index buffer never changes so can be safely read once
+	static TArray<int32> FaceIndices;
 };
 
 namespace ARKitUtil
diff --git a/Engine/Plugins/Runtime/AppleARKit/Source/AppleARKit/Public/AppleARKitFaceSupport.h b/Engine/Plugins/Runtime/AppleARKit/Source/AppleARKit/Public/AppleARKitFaceSupport.h
--- a/Engine/Plugins/Runtime/AppleARKit/Source/AppleARKit/Public/AppleARKitFaceSupport.h
+++ b/Engine/Plugins/Runtime/AppleARKit/Source/AppleARKit/Public/AppleARKitFaceSupport.h
@@ -12,46 +12,27 @@
 #endif
 
 class UARSessionConfig;
-class UARTrackedGeometry;
 class FAppleARKitConfiguration;
-class FARSystemBase;
+struct FAppleARKitAnchorData;
 
 
-class APPLEARKIT_API IAppleARKitFaceSupportCallback
-{
-public:
-	/** So that face processing can get access to the face geometry objects by their guid */
-	virtual UARTrackedGeometry* GetTrackedGeometry(const FGuid& GeoGuid) = 0;
-	/** So that face processing can add new geometry as updates come in */
-	virtual void AddTrackedGeometry(const FGuid& Guid, UARTrackedGeometry* TrackedGeo) = 0;
-};
-
 class APPLEARKIT_API IAppleARKitFaceSupport
 {
 public:
 #if SUPPORTS_ARKIT_1_0
 	/**
-	 * Forwards the anchor add for face processing
+	 * Converts a set of generic ARAnchors into their face anchor equivalents without exposing the main code to the face APIs
 	 *
-	 * @param NewAnchors the anchor list that the face support should process
-	 * @param AlignmentTransform the transform to apply for user alignment
-	 * @param FrameNumber the frame number to publish with LiveLink
-	 * @param Timestamp the timestamp to publish with LiveLink
-	 */
-	virtual void ProcessAnchorAdd(NSArray<ARAnchor*>* NewAnchors, const FTransform& AlignmentTransform, uint32 FrameNumber, double Timestamp) = 0;
-
-	/**
-	 * Forwards the anchor add for face processing
+	 * @param NewAnchors the list of anchors to convert to our intermediate format
+	 * @param Timestamp the timestamp of this update
+	 * @param FrameNumber the frame number for this update
 	 *
-	 * @param UpdatedAnchors the anchor list that the face support should process
-	 * @param AlignmentTransform the transform to apply for user alignment
-	 * @param FrameNumber the frame number to publish with LiveLink
-	 * @param Timestamp the timestamp to publish with LiveLink
+	 * @return the set of face anchors to dispatch
 	 */
-	virtual void ProcessAnchorUpdate(NSArray<ARAnchor*>* UpdatedAnchors, const FTransform& AlignmentTransform, uint32 FrameNumber, double Timestamp) = 0;
+	virtual TArray<TSharedPtr<FAppleARKitAnchorData>> MakeAnchorData(NSArray<ARAnchor*>* NewAnchors, double Timestamp, uint32 FrameNumber) = 0;
 
 	/**
-	 * Creates a face ar specific configuration object if that is requested
+	 * Creates a face ar specific configuration object if that is requested without exposing the main code to the face APIs
 	 *
 	 * @param SessionConfig the UE4 configuration object that needs processing
 	 * @param InConfiguration the legacy configuration object that needs to go away
@@ -60,65 +41,15 @@ public:
 #endif
 
 	IAppleARKitFaceSupport() { }
-	IAppleARKitFaceSupport(const TSharedRef<FARSystemBase, ESPMode::ThreadSafe>& InTrackingSystem, IAppleARKitFaceSupportCallback* Callback) { }
 	virtual ~IAppleARKitFaceSupport() { }
 };
 
-class APPLEARKIT_API FAppleARKitFaceSupportBase :
-	public IAppleARKitFaceSupport,
-	public TSharedFromThis<FAppleARKitFaceSupportBase, ESPMode::ThreadSafe>
-{
-public:
-#if SUPPORTS_ARKIT_1_0
-	/**
-	 * Forwards the anchor add for face processing
-	 *
-	 * @param NewAnchors the anchor list that the face support should process
-	 * @param AlignmentTransform the transform to apply for user alignment
-	 * @param FrameNumber the frame number to publish with LiveLink
-	 * @param Timestamp the timestamp to publish with LiveLink
-	 */
-	virtual void ProcessAnchorAdd(NSArray<ARAnchor*>* NewAnchors, const FTransform& AlignmentTransform, uint32 FrameNumber, double Timestamp) override
-	{
-
-	}
-
-	/**
-	 * Forwards the anchor add for face processing
-	 *
-	 * @param UpdatedAnchors the anchor list that the face support should process
-	 * @param AlignmentTransform the transform to apply for user alignment
-	 * @param FrameNumber the frame number to publish with LiveLink
-	 * @param Timestamp the timestamp to publish with LiveLink
-	 */
-	virtual void ProcessAnchorUpdate(NSArray<ARAnchor*>* UpdatedAnchors, const FTransform& AlignmentTransform, uint32 FrameNumber, double Timestamp) override
-	{
-
-	}
-
-	/**
-	 * Creates a face ar specific configuration object if that is requested
-	 *
-	 * @param SessionConfig the UE4 configuration object that needs processing
-	 * @param InConfiguration the legacy configuration object that needs to go away
-	 */
-	virtual ARConfiguration* ToARConfiguration(UARSessionConfig* SessionConfig, FAppleARKitConfiguration& InConfiguration) override
-	{
-		return nullptr;
-	}
-#endif
-
-	FAppleARKitFaceSupportBase() { }
-	FAppleARKitFaceSupportBase(const TSharedRef<FARSystemBase, ESPMode::ThreadSafe>& InTrackingSystem, IAppleARKitFaceSupportCallback* Callback) { }
-	virtual ~FAppleARKitFaceSupportBase() { }
-};
-
 class APPLEARKIT_API IAppleARKitFaceSupportFactory :
 	public IModularFeature
 {
 public:
 	/** Factory method that returns the object to use to handle face ar requests */
-	virtual TSharedPtr<FAppleARKitFaceSupportBase, ESPMode::ThreadSafe> CreateFaceSupport(const TSharedRef<FARSystemBase, ESPMode::ThreadSafe>& InTrackingSystem, IAppleARKitFaceSupportCallback* Callback) = 0;
+	virtual IAppleARKitFaceSupport* CreateFaceSupport() = 0;
 
 	static FName GetModularFeatureName()
 	{
diff --git a/Engine/Plugins/Runtime/AppleARKitFaceSupport/Source/AppleARKitFaceSupport/Private/AppleARKitFaceSupportImpl.cpp b/Engine/Plugins/Runtime/AppleARKitFaceSupport/Source/AppleARKitFaceSupport/Private/AppleARKitFaceSupportImpl.cpp
--- a/Engine/Plugins/Runtime/AppleARKitFaceSupport/Source/AppleARKitFaceSupport/Private/AppleARKitFaceSupportImpl.cpp
+++ b/Engine/Plugins/Runtime/AppleARKitFaceSupport/Source/AppleARKitFaceSupport/Private/AppleARKitFaceSupportImpl.cpp
@@ -10,28 +10,24 @@
 
 DECLARE_STATS_GROUP(TEXT("AppleARKitFaceSupport"), STATGROUP_APPLEARKITFACE, STATCAT_Advanced);
 
-uint32 LastTrackedFaceGeometryID = 0;
-
-TArray<int32> FAppleARKitFaceAnchorData::FaceIndices;
-
 #if SUPPORTS_ARKIT_1_0
 
-static TSharedPtr<FAppleARKitFaceAnchorData> MakeAnchorData(ARAnchor* Anchor)
+static TSharedPtr<FAppleARKitAnchorData> MakeAnchorData(ARAnchor* Anchor)
 {
-    TSharedPtr<FAppleARKitFaceAnchorData> NewAnchor;
+    TSharedPtr<FAppleARKitAnchorData> NewAnchor;
     if ([Anchor isKindOfClass:[ARFaceAnchor class]])
     {
         ARFaceAnchor* FaceAnchor = (ARFaceAnchor*)Anchor;
-        NewAnchor = MakeShared<FAppleARKitFaceAnchorData>(
+        NewAnchor = MakeShared<FAppleARKitAnchorData>(
 			FAppleARKitConversion::ToFGuid(FaceAnchor.identifier),
 			FAppleARKitConversion::ToFTransform(FaceAnchor.transform),
 			ToBlendShapeMap(FaceAnchor.blendShapes, FAppleARKitConversion::ToFTransform(FaceAnchor.transform)),
 			ToVertexBuffer(FaceAnchor.geometry.vertices, FaceAnchor.geometry.vertexCount)
         );
         // Only convert from 16bit to 32bit once
-        if (FAppleARKitFaceAnchorData::FaceIndices.Num() == 0)
+        if (FAppleARKitAnchorData::FaceIndices.Num() == 0)
         {
-            FAppleARKitFaceAnchorData::FaceIndices = To32BitIndexBuffer(FaceAnchor.geometry.triangleIndices, FaceAnchor.geometry.triangleCount * 3);
+            FAppleARKitAnchorData::FaceIndices = To32BitIndexBuffer(FaceAnchor.geometry.triangleIndices, FaceAnchor.geometry.triangleCount * 3);
         }
     }
 
@@ -40,14 +36,8 @@ static TSharedPtr<FAppleARKitFaceAnchorData> MakeAnchorData(ARAnchor* Anchor)
 
 #endif
 
-FAppleARKitFaceSupport::FAppleARKitFaceSupport(const TSharedRef<FARSystemBase, ESPMode::ThreadSafe>& InTrackingSystem, IAppleARKitFaceSupportCallback* InCallback)
-    : Callback(InCallback)
-	, ARSystem(InTrackingSystem)
-    , FrameNumber(0)
-    , Timestamp(0.0)
+FAppleARKitFaceSupport::FAppleARKitFaceSupport()
 {
-    check(Callback);
-
     // Create our LiveLink provider if the project setting is enabled
     if (GetDefault<UAppleARKitSettings>()->bEnableLiveLinkForFaceTracking)
     {
@@ -88,95 +78,25 @@ ARConfiguration* FAppleARKitFaceSupport::ToARConfiguration(UARSessionConfig* Ses
 	return SessionConfiguration;
 }
 
-void FAppleARKitFaceSupport::ProcessAnchorAdd(NSArray<ARAnchor*>* NewAnchors, const FTransform& InAlignmentTransform, uint32 InFrameNumber, double InTimestamp)
+TArray<TSharedPtr<FAppleARKitAnchorData>> FAppleARKitFaceSupport::MakeAnchorData(NSArray<ARAnchor*>* Anchors, double Timestamp, uint32 FrameNumber)
 {
-	DECLARE_CYCLE_STAT(TEXT("FaceAR::ProcessAnchorAdd"), STAT_FaceAR_ProcessAnchorAdd, STATGROUP_APPLEARKITFACE);
+	TArray<TSharedPtr<FAppleARKitAnchorData>> AnchorList;
 
-    FrameNumber = InFrameNumber;
-    Timestamp = InTimestamp;
-
-	for (ARAnchor* Anchor in NewAnchors)
+	for (ARAnchor* Anchor in Anchors)
 	{
-		TSharedPtr<FAppleARKitFaceAnchorData> AnchorData = MakeAnchorData(Anchor);
+		TSharedPtr<FAppleARKitAnchorData> AnchorData = ::MakeAnchorData(Anchor);
 		if (AnchorData.IsValid())
-		{
-			auto AnchorTask = FSimpleDelegateGraphTask::FDelegate::CreateSP(this, &FAppleARKitFaceSupport::ProcessAnchorAdd_Internal, AnchorData.ToSharedRef());
-			FSimpleDelegateGraphTask::CreateAndDispatchWhenReady(AnchorTask, GET_STATID(STAT_FaceAR_ProcessAnchorAdd), nullptr, ENamedThreads::GameThread);
-		}
-	}
-}
-
-void FAppleARKitFaceSupport::ProcessAnchorUpdate(NSArray<ARAnchor*>* UpdatedAnchors, const FTransform& InAlignmentTransform, uint32 InFrameNumber, double InTimestamp)
-{
-	DECLARE_CYCLE_STAT(TEXT("FaceAR::ProcessAnchorUpdate"), STAT_FaceAR_ProcessAnchorUpdate, STATGROUP_APPLEARKITFACE);
-
-    FrameNumber = InFrameNumber;
-    Timestamp = InTimestamp;
-
-	for (ARAnchor* Anchor in UpdatedAnchors)
-	{
-		TSharedPtr<FAppleARKitFaceAnchorData> AnchorData = MakeAnchorData(Anchor);
-		if (AnchorData.IsValid())
-		{
-			auto AnchorTask = FSimpleDelegateGraphTask::FDelegate::CreateSP(this, &FAppleARKitFaceSupport::ProcessAnchorUpdate_Internal, AnchorData.ToSharedRef());
-			FSimpleDelegateGraphTask::CreateAndDispatchWhenReady(AnchorTask, GET_STATID(STAT_FaceAR_ProcessAnchorUpdate), nullptr, ENamedThreads::GameThread);
-		}
-	}
-}
-
-void FAppleARKitFaceSupport::ProcessAnchorAdd_Internal(TSharedRef<FAppleARKitFaceAnchorData> AnchorData)
-{
-	if (AnchorData->AnchorType == EAppleAnchorType::FaceAnchor)
-	{
-        // Update LiveLink first, because the other updates use MoveTemp for efficiency
-        if (LiveLinkSource.IsValid())
-        {
-            LiveLinkSource->PublishBlendShapes(FaceTrackingLiveLinkSubjectName, Timestamp, FrameNumber, AnchorData->BlendShapes);
-        }
-
-        FString NewAnchorDebugName;
-        UARFaceGeometry* NewGeo = NewObject<UARFaceGeometry>();
-        NewAnchorDebugName = FString::Printf(TEXT("FACE-%02d"), LastTrackedFaceGeometryID++);
-        NewGeo->SetDebugName(FName(*NewAnchorDebugName));
-
-        NewGeo->UpdateTrackedGeometry(ARSystem, FrameNumber, Timestamp, AnchorData->Transform, AlignmentTransform, AnchorData->BlendShapes, AnchorData->FaceVerts, AnchorData->FaceIndices);
-
-        // Add it to the ar system since that's what other APIs query against
-        Callback->AddTrackedGeometry(AnchorData->AnchorGUID, NewGeo);
-    }
-}
-
-void FAppleARKitFaceSupport::ProcessAnchorUpdate_Internal(TSharedRef<FAppleARKitFaceAnchorData> AnchorData)
-{
-	UARTrackedGeometry* FoundGeometry = Callback->GetTrackedGeometry(AnchorData->AnchorGUID);
-	if (ensure(FoundGeometry != nullptr))
-	{
-		if (UARFaceGeometry* FaceGeo = Cast<UARFaceGeometry>(FoundGeometry))
 		{
 			// Update LiveLink first, because the other updates use MoveTemp for efficiency
 			if (LiveLinkSource.IsValid())
 			{
 				LiveLinkSource->PublishBlendShapes(FaceTrackingLiveLinkSubjectName, Timestamp, FrameNumber, AnchorData->BlendShapes);
 			}
-
-			const TArray<UARPin*>& Pins = ARSystem->GetAllPins();
-			TArray<UARPin*> PinsToUpdate = ARKitUtil::PinsFromGeometry(FoundGeometry, Pins);
-
-			// We figure out the delta transform for the Anchor (aka. TrackedGeometry in ARKit) and apply that
-			// delta to figure out the new ARPin transform.
-			const FTransform Anchor_LocalToTrackingTransform_PreUpdate = FoundGeometry->GetLocalToTrackingTransform_NoAlignment();
-			const FTransform& Anchor_LocalToTrackingTransform_PostUpdate = AnchorData->Transform;
-
-			const FTransform AnchorDeltaTransform = Anchor_LocalToTrackingTransform_PreUpdate.GetRelativeTransform(Anchor_LocalToTrackingTransform_PostUpdate);
-
-			FaceGeo->UpdateTrackedGeometry(ARSystem, FrameNumber, Timestamp, AnchorData->Transform, AlignmentTransform, AnchorData->BlendShapes, AnchorData->FaceVerts, AnchorData->FaceIndices);
-			for (UARPin* Pin : PinsToUpdate)
-			{
-				const FTransform Pin_LocalToTrackingTransform_PostUpdate = Pin->GetLocalToTrackingTransform_NoAlignment() * AnchorDeltaTransform;
-				Pin->OnTransformUpdated(Pin_LocalToTrackingTransform_PostUpdate);
-			}
+			AnchorList.Add(AnchorData);
 		}
 	}
+
+	return AnchorList;
 }
 
 #endif
diff --git a/Engine/Plugins/Runtime/AppleARKitFaceSupport/Source/AppleARKitFaceSupport/Private/AppleARKitFaceSupportImpl.h b/Engine/Plugins/Runtime/AppleARKitFaceSupport/Source/AppleARKitFaceSupport/Private/AppleARKitFaceSupportImpl.h
--- a/Engine/Plugins/Runtime/AppleARKitFaceSupport/Source/AppleARKitFaceSupport/Private/AppleARKitFaceSupportImpl.h
+++ b/Engine/Plugins/Runtime/AppleARKitFaceSupport/Source/AppleARKitFaceSupport/Private/AppleARKitFaceSupportImpl.h
@@ -8,48 +8,21 @@
 #include "AppleARKitConversion.h"
 
 
-struct FAppleARKitFaceAnchorData :
-	public FAppleARKitAnchorData
-{
-	FAppleARKitFaceAnchorData(FGuid InAnchorGuid, FTransform InTransform, FARBlendShapeMap InBlendShapes, TArray<FVector> InFaceVerts)
-		: FAppleARKitAnchorData(InAnchorGuid, InTransform)
-		, BlendShapes( MoveTemp(InBlendShapes) )
-		, FaceVerts( MoveTemp(InFaceVerts) )
-	{
-	}
-
-	FARBlendShapeMap BlendShapes;
-	TArray<FVector> FaceVerts;
-	// Note: the index buffer never changes so can be safely read once
-	static TArray<int32> FaceIndices;
-};
-
 class APPLEARKITFACESUPPORT_API FAppleARKitFaceSupport :
-	public FAppleARKitFaceSupportBase
+	public IAppleARKitFaceSupport
 {
 public:
-	FAppleARKitFaceSupport(const TSharedRef<FARSystemBase, ESPMode::ThreadSafe>& InTrackingSystem, IAppleARKitFaceSupportCallback* Callback);
+	FAppleARKitFaceSupport();
 	virtual ~FAppleARKitFaceSupport();
 
 private:
 #if SUPPORTS_ARKIT_1_0
 	// ~IAppleARKitFaceSupport
-	virtual void ProcessAnchorAdd(NSArray<ARAnchor*>* NewAnchors, const FTransform& AlignmentTransform, uint32 FrameNumber, double Timestamp) override;
-	virtual void ProcessAnchorUpdate(NSArray<ARAnchor*>* UpdatedAnchors, const FTransform& AlignmentTransform, uint32 FrameNumber, double Timestamp) override;
+	virtual TArray<TSharedPtr<FAppleARKitAnchorData>> MakeAnchorData(NSArray<ARAnchor*>* NewAnchors, double Timestamp, uint32 FrameNumber) override;
 	virtual ARConfiguration* ToARConfiguration(UARSessionConfig* InSessionConfig, FAppleARKitConfiguration& InConfiguration) override;
 	// ~IAppleARKitFaceSupport
-
-	void ProcessAnchorAdd_Internal(TSharedRef<FAppleARKitFaceAnchorData> AnchorData);
-	void ProcessAnchorUpdate_Internal(TSharedRef<FAppleARKitFaceAnchorData> AnchorData);
 #endif
     
-    IAppleARKitFaceSupportCallback* Callback;
-	const TSharedRef<FARSystemBase, ESPMode::ThreadSafe> ARSystem;
-
-	uint32 FrameNumber;
-	double Timestamp;
-	FTransform AlignmentTransform;
-
 	/** If requested, publishes face ar updates to LiveLink for the animation system to use */
 	TSharedPtr<ILiveLinkSourceARKit> LiveLinkSource;
 	/** Copied from the UARSessionConfig project settings object */
diff --git a/Engine/Plugins/Runtime/AppleARKitFaceSupport/Source/AppleARKitFaceSupport/Private/AppleARKitFaceSupportModule.cpp b/Engine/Plugins/Runtime/AppleARKitFaceSupport/Source/AppleARKitFaceSupport/Private/AppleARKitFaceSupportModule.cpp
--- a/Engine/Plugins/Runtime/AppleARKitFaceSupport/Source/AppleARKitFaceSupport/Private/AppleARKitFaceSupportModule.cpp
+++ b/Engine/Plugins/Runtime/AppleARKitFaceSupport/Source/AppleARKitFaceSupport/Private/AppleARKitFaceSupportModule.cpp
@@ -31,24 +31,27 @@ public:
 	}
 
 private:
-	FAppleARKitFaceSupportFactory()
+	FAppleARKitFaceSupportFactory() :
+		FaceARSupport(nullptr)
 	{
 	}
+
 	virtual ~FAppleARKitFaceSupportFactory()
 	{
+		delete FaceARSupport;
 	}
 
-	virtual TSharedPtr<FAppleARKitFaceSupportBase, ESPMode::ThreadSafe> CreateFaceSupport(const TSharedRef<FARSystemBase, ESPMode::ThreadSafe>& InTrackingSystem, IAppleARKitFaceSupportCallback* Callback)
+	virtual IAppleARKitFaceSupport* CreateFaceSupport()
 	{
-		if (!FaceARSupport.IsValid())
+		if (FaceARSupport == nullptr)
 		{
-			FaceARSupport = MakeShared<FAppleARKitFaceSupport, ESPMode::ThreadSafe>(InTrackingSystem, Callback);
+			FaceARSupport = new FAppleARKitFaceSupport();
 		}
 		return FaceARSupport;
 	};
 
 	static FAppleARKitFaceSupportFactory* Factory;
-	TSharedPtr<FAppleARKitFaceSupport, ESPMode::ThreadSafe> FaceARSupport;
+	IAppleARKitFaceSupport* FaceARSupport;
 };
 
 FAppleARKitFaceSupportFactory* FAppleARKitFaceSupportFactory::Factory = nullptr;
