A different approach to the compiler generating bogus code

diff --git a/Engine/Plugins/Runtime/AppleARKit/Source/AppleARKit/Private/AppleARKitSystem.cpp b/Engine/Plugins/Runtime/AppleARKit/Source/AppleARKit/Private/AppleARKitSystem.cpp
--- a/Engine/Plugins/Runtime/AppleARKit/Source/AppleARKit/Private/AppleARKitSystem.cpp
+++ b/Engine/Plugins/Runtime/AppleARKit/Source/AppleARKit/Private/AppleARKitSystem.cpp
@@ -193,11 +193,10 @@ void FAppleARKitSystem::CheckForFaceARSupport(UARSessionConfig* InSessionConfig)
 	}
 
 	// We need to get the face support from the factory method, which is a modular feature to avoid dependencies
-	IModularFeature* Feature = IModularFeatures::Get().GetModularFeatureImplementation("AppleARKitFaceSupportFactory", 0);
-	if (ensureAlwaysMsgf(Feature != nullptr, TEXT("Face AR session has been requested but the face ar plugin is not enabled")))
+	TArray<IAppleARKitFaceSupport*> Impls = IModularFeatures::Get().GetModularFeatureImplementations<IAppleARKitFaceSupport>("AppleARKitFaceSupport");
+	if (ensureAlwaysMsgf(Impls.Num() > 0, TEXT("Face AR session has been requested but the face ar plugin is not enabled")))
 	{
-		IAppleARKitFaceSupportFactory* Factory = static_cast<IAppleARKitFaceSupportFactory*>(Feature);
-		FaceARSupport = Factory->CreateFaceSupport();
+		FaceARSupport = Impls[0];
 		ensureAlwaysMsgf(FaceARSupport != nullptr, TEXT("Face AR session has been requested but the face ar plugin is not enabled"));
 	}
 }
@@ -819,6 +818,7 @@ void FAppleARKitSystem::SetDeviceOrientation( EScreenOrientation::Type InOrienta
 }
 
 
+PRAGMA_DISABLE_OPTIMIZATION
 bool FAppleARKitSystem::Run(UARSessionConfig* SessionConfig)
 {
 	if (IsRunning())
@@ -901,6 +901,7 @@ bool FAppleARKitSystem::Run(UARSessionConfig* SessionConfig)
 	OnARSessionStarted.Broadcast();
 	return true;
 }
+PRAGMA_ENABLE_OPTIMIZATION
 
 bool FAppleARKitSystem::IsRunning() const
 {
diff --git a/Engine/Plugins/Runtime/AppleARKit/Source/AppleARKit/Public/AppleARKitFaceSupport.h b/Engine/Plugins/Runtime/AppleARKit/Source/AppleARKit/Public/AppleARKitFaceSupport.h
--- a/Engine/Plugins/Runtime/AppleARKit/Source/AppleARKit/Public/AppleARKitFaceSupport.h
+++ b/Engine/Plugins/Runtime/AppleARKit/Source/AppleARKit/Public/AppleARKitFaceSupport.h
@@ -15,7 +15,8 @@ class UARSessionConfig;
 struct FAppleARKitAnchorData;
 
 
-class APPLEARKIT_API IAppleARKitFaceSupport
+class APPLEARKIT_API IAppleARKitFaceSupport :
+	public IModularFeature
 {
 public:
 #if SUPPORTS_ARKIT_1_0
@@ -49,18 +50,10 @@ public:
 	 */
 	virtual ARConfiguration* ToARConfiguration(UARSessionConfig* SessionConfig) = 0;
 #endif
-};
-
-class APPLEARKIT_API IAppleARKitFaceSupportFactory :
-	public IModularFeature
-{
-public:
-	/** Factory method that returns the object to use to handle face ar requests */
-	virtual IAppleARKitFaceSupport* CreateFaceSupport() = 0;
 
 	static FName GetModularFeatureName()
 	{
-		static FName FeatureName = FName(TEXT("AppleARKitFaceSupportFactory"));
+		static FName FeatureName = FName(TEXT("AppleARKitFaceSupport"));
 		return FeatureName;
 	}
 };
diff --git a/Engine/Plugins/Runtime/AppleARKitFaceSupport/Source/AppleARKitFaceSupport/Private/AppleARKitFaceSupportImpl.cpp b/Engine/Plugins/Runtime/AppleARKitFaceSupport/Source/AppleARKitFaceSupport/Private/AppleARKitFaceSupportImpl.cpp
--- a/Engine/Plugins/Runtime/AppleARKitFaceSupport/Source/AppleARKitFaceSupport/Private/AppleARKitFaceSupportImpl.cpp
+++ b/Engine/Plugins/Runtime/AppleARKitFaceSupport/Source/AppleARKitFaceSupport/Private/AppleARKitFaceSupportImpl.cpp
@@ -38,17 +38,6 @@ static TSharedPtr<FAppleARKitAnchorData> MakeAnchorData(ARAnchor* Anchor)
 
 FAppleARKitFaceSupport::FAppleARKitFaceSupport()
 {
-    // Create our LiveLink provider if the project setting is enabled
-    if (GetDefault<UAppleARKitSettings>()->bEnableLiveLinkForFaceTracking)
-    {
-        FaceTrackingLiveLinkSubjectName = GetDefault<UAppleARKitSettings>()->DefaultFaceTrackingLiveLinkSubjectName;
-#if PLATFORM_IOS
-        LiveLinkSource = FAppleARKitLiveLinkSourceFactory::CreateLiveLinkSource(true);
-#else
-        // This should be started already, but just in case
-        FAppleARKitLiveLinkSourceFactory::CreateLiveLinkRemoteListener();
-#endif
-    }
 }
 
 FAppleARKitFaceSupport::~FAppleARKitFaceSupport()
@@ -97,6 +86,23 @@ TArray<TSharedPtr<FAppleARKitAnchorData>> FAppleARKitFaceSupport::MakeAnchorData
 
 void FAppleARKitFaceSupport::PublishLiveLinkData(TSharedPtr<FAppleARKitAnchorData> Anchor, double Timestamp, uint32 FrameNumber)
 {
+	static bool bNeedsInit = true;
+	if (bNeedsInit)
+	{
+		bNeedsInit = false;
+	    // Create our LiveLink provider if the project setting is enabled
+		if (GetDefault<UAppleARKitSettings>()->bEnableLiveLinkForFaceTracking)
+		{
+			FaceTrackingLiveLinkSubjectName = GetDefault<UAppleARKitSettings>()->DefaultFaceTrackingLiveLinkSubjectName;
+	#if PLATFORM_IOS
+			LiveLinkSource = FAppleARKitLiveLinkSourceFactory::CreateLiveLinkSource(true);
+	#else
+			// This should be started already, but just in case
+			FAppleARKitLiveLinkSourceFactory::CreateLiveLinkRemoteListener();
+	#endif
+		}
+	}
+	
 	if (LiveLinkSource.IsValid())
 	{
         LiveLinkSource->PublishBlendShapes(FaceTrackingLiveLinkSubjectName, Timestamp, FrameNumber, Anchor->BlendShapes);
diff --git a/Engine/Plugins/Runtime/AppleARKitFaceSupport/Source/AppleARKitFaceSupport/Private/AppleARKitFaceSupportModule.cpp b/Engine/Plugins/Runtime/AppleARKitFaceSupport/Source/AppleARKitFaceSupport/Private/AppleARKitFaceSupportModule.cpp
--- a/Engine/Plugins/Runtime/AppleARKitFaceSupport/Source/AppleARKitFaceSupport/Private/AppleARKitFaceSupportModule.cpp
+++ b/Engine/Plugins/Runtime/AppleARKitFaceSupport/Source/AppleARKitFaceSupport/Private/AppleARKitFaceSupportModule.cpp
@@ -13,54 +13,14 @@ IMPLEMENT_MODULE(FAppleARKitFaceSupportModule, AppleARKitFaceSupport);
 
 DEFINE_LOG_CATEGORY(LogAppleARKitFace);
 
-class APPLEARKITFACESUPPORT_API FAppleARKitFaceSupportFactory :
-	public IAppleARKitFaceSupportFactory
-{
-public:
-	static void CreateInstance()
-	{
-		Factory = new FAppleARKitFaceSupportFactory();
-		IModularFeatures::Get().RegisterModularFeature(Factory->GetModularFeatureName(), Factory);
-	}
-
-	static void DestroyInstance()
-	{
-		IModularFeatures::Get().UnregisterModularFeature(Factory->GetModularFeatureName(), Factory);
-		delete Factory;
-		Factory = nullptr;
-	}
-
-private:
-	FAppleARKitFaceSupportFactory() :
-		FaceARSupport(nullptr)
-	{
-	}
-
-	virtual ~FAppleARKitFaceSupportFactory()
-	{
-		delete FaceARSupport;
-	}
-
-	virtual IAppleARKitFaceSupport* CreateFaceSupport()
-	{
-		if (FaceARSupport == nullptr)
-		{
-			FaceARSupport = new FAppleARKitFaceSupport();
-		}
-		return FaceARSupport;
-	};
-
-	static FAppleARKitFaceSupportFactory* Factory;
-	IAppleARKitFaceSupport* FaceARSupport;
-};
-
-FAppleARKitFaceSupportFactory* FAppleARKitFaceSupportFactory::Factory = nullptr;
+FAppleARKitFaceSupport* FaceSupportInstance = nullptr;
 
 void FAppleARKitFaceSupportModule::StartupModule()
 {
 	ensureMsgf(FModuleManager::Get().LoadModule("AppleARKit"), TEXT("ARKitFaceSupport depends on the AppleARKit module."));
 
-	FAppleARKitFaceSupportFactory::CreateInstance();
+	FaceSupportInstance = new FAppleARKitFaceSupport();
+	IModularFeatures::Get().RegisterModularFeature(FaceSupportInstance->GetModularFeatureName(), FaceSupportInstance);
 
 	// LiveLink listener needs to be created here so that the editor can receive remote publishing events
 #if PLATFORM_DESKTOP
@@ -75,6 +35,8 @@ void FAppleARKitFaceSupportModule::StartupModule()
 
 void FAppleARKitFaceSupportModule::ShutdownModule()
 {
-	FAppleARKitFaceSupportFactory::DestroyInstance();
+	IModularFeatures::Get().UnregisterModularFeature(FaceSupportInstance->GetModularFeatureName(), FaceSupportInstance);
+	delete FaceSupportInstance;
+	FaceSupportInstance = nullptr;
 }
 
