Switched the face support to using shared pointers instead

diff --git a/Engine/Plugins/Runtime/AppleARKit/Source/AppleARKit/Public/AppleARKitFaceSupport.h b/Engine/Plugins/Runtime/AppleARKit/Source/AppleARKit/Public/AppleARKitFaceSupport.h
--- a/Engine/Plugins/Runtime/AppleARKit/Source/AppleARKit/Public/AppleARKitFaceSupport.h
+++ b/Engine/Plugins/Runtime/AppleARKit/Source/AppleARKit/Public/AppleARKitFaceSupport.h
@@ -29,7 +29,7 @@ public:
 	 *
 	 * @return the set of face anchors to dispatch
 	 */
-	virtual TArray<TSharedPtr<FAppleARKitAnchorData>> MakeAnchorData(NSArray<ARAnchor*>* NewAnchors, double Timestamp, uint32 FrameNumber) = 0;
+	virtual TArray<TSharedPtr<FAppleARKitAnchorData>> MakeAnchorData(NSArray<ARAnchor*>* NewAnchors, double Timestamp, uint32 FrameNumber) { return TArray<TSharedPtr<FAppleARKitAnchorData>>(); }
 
 	/**
 	 * Publishes any face AR data that needs to be sent to LiveLink. Done as a separate step because MakeAnchorData is called
@@ -41,14 +41,14 @@ public:
 	 *
 	 * @return the set of face anchors to dispatch
 	 */
-	virtual void PublishLiveLinkData(TSharedPtr<FAppleARKitAnchorData> Anchor, double Timestamp, uint32 FrameNumber) = 0;
+	virtual void PublishLiveLinkData(TSharedPtr<FAppleARKitAnchorData> Anchor, double Timestamp, uint32 FrameNumber) { }
 
 	/**
 	 * Creates a face ar specific configuration object if that is requested without exposing the main code to the face APIs
 	 *
 	 * @param SessionConfig the UE4 configuration object that needs processing
 	 */
-	virtual ARConfiguration* ToARConfiguration(UARSessionConfig* SessionConfig) = 0;
+	virtual ARConfiguration* ToARConfiguration(UARSessionConfig* SessionConfig) { return nullptr; }
 #endif
 
 	static FName GetModularFeatureName()
diff --git a/Engine/Plugins/Runtime/AppleARKitFaceSupport/Source/AppleARKitFaceSupport/Private/AppleARKitFaceSupportImpl.cpp b/Engine/Plugins/Runtime/AppleARKitFaceSupport/Source/AppleARKitFaceSupport/Private/AppleARKitFaceSupportImpl.cpp
--- a/Engine/Plugins/Runtime/AppleARKitFaceSupport/Source/AppleARKitFaceSupport/Private/AppleARKitFaceSupportImpl.cpp
+++ b/Engine/Plugins/Runtime/AppleARKitFaceSupport/Source/AppleARKitFaceSupport/Private/AppleARKitFaceSupportImpl.cpp
@@ -46,6 +46,16 @@ FAppleARKitFaceSupport::~FAppleARKitFaceSupport()
 	check(GIsRequestingExit);
 }
 
+void FAppleARKitFaceSupport::Init()
+{
+	IModularFeatures::Get().RegisterModularFeature(GetModularFeatureName(), this);
+}
+
+void FAppleARKitFaceSupport::Shutdown()
+{
+	IModularFeatures::Get().UnregisterModularFeature(GetModularFeatureName(), this);
+}
+
 #if SUPPORTS_ARKIT_1_0
 
 ARConfiguration* FAppleARKitFaceSupport::ToARConfiguration(UARSessionConfig* SessionConfig)
@@ -102,7 +112,7 @@ void FAppleARKitFaceSupport::PublishLiveLinkData(TSharedPtr<FAppleARKitAnchorDat
 	#endif
 		}
 	}
-	
+
 	if (LiveLinkSource.IsValid())
 	{
         LiveLinkSource->PublishBlendShapes(FaceTrackingLiveLinkSubjectName, Timestamp, FrameNumber, Anchor->BlendShapes);
diff --git a/Engine/Plugins/Runtime/AppleARKitFaceSupport/Source/AppleARKitFaceSupport/Private/AppleARKitFaceSupportImpl.h b/Engine/Plugins/Runtime/AppleARKitFaceSupport/Source/AppleARKitFaceSupport/Private/AppleARKitFaceSupportImpl.h
--- a/Engine/Plugins/Runtime/AppleARKitFaceSupport/Source/AppleARKitFaceSupport/Private/AppleARKitFaceSupportImpl.h
+++ b/Engine/Plugins/Runtime/AppleARKitFaceSupport/Source/AppleARKitFaceSupport/Private/AppleARKitFaceSupportImpl.h
@@ -9,12 +9,16 @@
 
 
 class APPLEARKITFACESUPPORT_API FAppleARKitFaceSupport :
-	public IAppleARKitFaceSupport
+	public IAppleARKitFaceSupport,
+	public TSharedFromThis<FAppleARKitFaceSupport, ESPMode::ThreadSafe>
 {
 public:
 	FAppleARKitFaceSupport();
 	virtual ~FAppleARKitFaceSupport();
 
+	void Init();
+	void Shutdown();
+
 private:
 #if SUPPORTS_ARKIT_1_0
 	// ~IAppleARKitFaceSupport
diff --git a/Engine/Plugins/Runtime/AppleARKitFaceSupport/Source/AppleARKitFaceSupport/Private/AppleARKitFaceSupportModule.cpp b/Engine/Plugins/Runtime/AppleARKitFaceSupport/Source/AppleARKitFaceSupport/Private/AppleARKitFaceSupportModule.cpp
--- a/Engine/Plugins/Runtime/AppleARKitFaceSupport/Source/AppleARKitFaceSupport/Private/AppleARKitFaceSupportModule.cpp
+++ b/Engine/Plugins/Runtime/AppleARKitFaceSupport/Source/AppleARKitFaceSupport/Private/AppleARKitFaceSupportModule.cpp
@@ -13,14 +13,14 @@ IMPLEMENT_MODULE(FAppleARKitFaceSupportModule, AppleARKitFaceSupport);
 
 DEFINE_LOG_CATEGORY(LogAppleARKitFace);
 
-FAppleARKitFaceSupport* FaceSupportInstance = nullptr;
+TSharedPtr<FAppleARKitFaceSupport, ESPMode::ThreadSafe> FaceSupportInstance;
 
 void FAppleARKitFaceSupportModule::StartupModule()
 {
 	ensureMsgf(FModuleManager::Get().LoadModule("AppleARKit"), TEXT("ARKitFaceSupport depends on the AppleARKit module."));
 
-	FaceSupportInstance = new FAppleARKitFaceSupport();
-	IModularFeatures::Get().RegisterModularFeature(FaceSupportInstance->GetModularFeatureName(), FaceSupportInstance);
+	FaceSupportInstance = MakeShared<FAppleARKitFaceSupport, ESPMode::ThreadSafe>();
+	FaceSupportInstance->Init();
 
 	// LiveLink listener needs to be created here so that the editor can receive remote publishing events
 #if PLATFORM_DESKTOP
@@ -35,8 +35,7 @@ void FAppleARKitFaceSupportModule::StartupModule()
 
 void FAppleARKitFaceSupportModule::ShutdownModule()
 {
-	IModularFeatures::Get().UnregisterModularFeature(FaceSupportInstance->GetModularFeatureName(), FaceSupportInstance);
-	delete FaceSupportInstance;
+	FaceSupportInstance->Shutdown();
 	FaceSupportInstance = nullptr;
 }
 
