Added stats for face ar

diff --git a/Engine/Plugins/Runtime/AppleARKitFaceSupport/Source/AppleARKitFaceSupport/Private/AppleARKitFaceMeshComponent.cpp b/Engine/Plugins/Runtime/AppleARKitFaceSupport/Source/AppleARKitFaceSupport/Private/AppleARKitFaceMeshComponent.cpp
--- a/Engine/Plugins/Runtime/AppleARKitFaceSupport/Source/AppleARKitFaceSupport/Private/AppleARKitFaceMeshComponent.cpp
+++ b/Engine/Plugins/Runtime/AppleARKitFaceSupport/Source/AppleARKitFaceSupport/Private/AppleARKitFaceMeshComponent.cpp
@@ -2,6 +2,7 @@
 
 #include "AppleARKitFaceMeshComponent.h"
 #include "ARBlueprintLibrary.h"
+#include "AppleARKitFaceSupportModule.h"
 
 #if SUPPORTS_ARKIT_1_0
 	#import <ARKit/ARKit.h>
@@ -13,6 +14,7 @@
 #include "AppleARKitSettings.h"
 #include "AppleARKitConversion.h"
 
+DECLARE_CYCLE_STAT(TEXT("Update"), STAT_FaceAR_Component_Update, STATGROUP_FaceAR);
 
 #if SUPPORTS_ARKIT_1_0
 
@@ -316,6 +318,8 @@ void UAppleARKitFaceMeshComponent::SetAutoBind(bool bAutoBind)
 
 void UAppleARKitFaceMeshComponent::TickComponent(float DeltaTime, ELevelTick, FActorComponentTickFunction*)
 {
+	SCOPE_CYCLE_COUNTER(STAT_FaceAR_Component_Update);
+
 	if (!bAutoBindToLocalFaceMesh)
 	{
 		return;
diff --git a/Engine/Plugins/Runtime/AppleARKitFaceSupport/Source/AppleARKitFaceSupport/Private/AppleARKitFaceSupportImpl.cpp b/Engine/Plugins/Runtime/AppleARKitFaceSupport/Source/AppleARKitFaceSupport/Private/AppleARKitFaceSupportImpl.cpp
--- a/Engine/Plugins/Runtime/AppleARKitFaceSupport/Source/AppleARKitFaceSupport/Private/AppleARKitFaceSupportImpl.cpp
+++ b/Engine/Plugins/Runtime/AppleARKitFaceSupport/Source/AppleARKitFaceSupport/Private/AppleARKitFaceSupportImpl.cpp
@@ -8,12 +8,14 @@
 #include "ARSystem.h"
 #include "Misc/ConfigCacheIni.h"
 
-DECLARE_STATS_GROUP(TEXT("AppleARKitFaceSupport"), STATGROUP_APPLEARKITFACE, STATCAT_Advanced);
+DECLARE_CYCLE_STAT(TEXT("Conversion"), STAT_FaceAR_Conversion, STATGROUP_FaceAR);
 
 #if SUPPORTS_ARKIT_1_0
 
 static TSharedPtr<FAppleARKitAnchorData> MakeAnchorData(ARAnchor* Anchor, const FRotator& AdjustBy)
 {
+	SCOPE_CYCLE_COUNTER(STAT_FaceAR_Conversion);
+	
     TSharedPtr<FAppleARKitAnchorData> NewAnchor;
     if ([Anchor isKindOfClass:[ARFaceAnchor class]])
     {
diff --git a/Engine/Plugins/Runtime/AppleARKitFaceSupport/Source/AppleARKitFaceSupport/Private/AppleARKitFaceSupportModule.h b/Engine/Plugins/Runtime/AppleARKitFaceSupport/Source/AppleARKitFaceSupport/Private/AppleARKitFaceSupportModule.h
--- a/Engine/Plugins/Runtime/AppleARKitFaceSupport/Source/AppleARKitFaceSupport/Private/AppleARKitFaceSupportModule.h
+++ b/Engine/Plugins/Runtime/AppleARKitFaceSupport/Source/AppleARKitFaceSupport/Private/AppleARKitFaceSupportModule.h
@@ -2,6 +2,7 @@
 
 #pragma once
 
+#include "Stats/Stats.h"
 #include "Modules/ModuleManager.h"
 
 class APPLEARKITFACESUPPORT_API FAppleARKitFaceSupportModule :
@@ -15,3 +16,4 @@ public:
 
 DECLARE_LOG_CATEGORY_EXTERN(LogAppleARKitFace, Log, All);
 
+DECLARE_STATS_GROUP(TEXT("Face AR"), STATGROUP_FaceAR, STATCAT_Advanced);
diff --git a/Engine/Plugins/Runtime/AppleARKitFaceSupport/Source/AppleARKitFaceSupport/Private/AppleARKitLiveLinkSource.cpp b/Engine/Plugins/Runtime/AppleARKitFaceSupport/Source/AppleARKitFaceSupport/Private/AppleARKitLiveLinkSource.cpp
--- a/Engine/Plugins/Runtime/AppleARKitFaceSupport/Source/AppleARKitFaceSupport/Private/AppleARKitLiveLinkSource.cpp
+++ b/Engine/Plugins/Runtime/AppleARKitFaceSupport/Source/AppleARKitFaceSupport/Private/AppleARKitLiveLinkSource.cpp
@@ -19,6 +19,19 @@
 
 #include "AppleARKitFaceSupportModule.h"
 
+DECLARE_CYCLE_STAT(TEXT("Publish Local LiveLink"), STAT_FaceAR_Local_PublishLiveLink, STATGROUP_FaceAR);
+DECLARE_CYCLE_STAT(TEXT("Publish Remote LiveLink"), STAT_FaceAR_Remote_PublishLiveLink, STATGROUP_FaceAR);
+DECLARE_CYCLE_STAT(TEXT("Receive LiveLink"), STAT_FaceAR_ReceiveLiveLink, STATGROUP_FaceAR);
+
+DECLARE_DWORD_ACCUMULATOR_STAT(TEXT("Total LiveLink Bytes Sent"), STAT_FaceAR_Total_LiveLink_BytesSent, STATGROUP_FaceAR);
+DECLARE_DWORD_ACCUMULATOR_STAT(TEXT("Total LiveLink Packets Sent"), STAT_FaceAR_Total_LiveLink_PacketsSent, STATGROUP_FaceAR);
+DECLARE_DWORD_ACCUMULATOR_STAT(TEXT("Total LiveLink Bytes Recv"), STAT_FaceAR_Total_LiveLink_BytesRecv, STATGROUP_FaceAR);
+DECLARE_DWORD_ACCUMULATOR_STAT(TEXT("Total LiveLink Packets Recv"), STAT_FaceAR_Total_LiveLink_PacketsRecv, STATGROUP_FaceAR);
+
+DECLARE_DWORD_COUNTER_STAT(TEXT("LiveLink Bytes Sent"), STAT_FaceAR_LiveLink_BytesSent, STATGROUP_FaceAR);
+DECLARE_DWORD_COUNTER_STAT(TEXT("LiveLink Packets Sent"), STAT_FaceAR_LiveLink_PacketsSent, STATGROUP_FaceAR);
+DECLARE_DWORD_COUNTER_STAT(TEXT("LiveLink Bytes Recv"), STAT_FaceAR_LiveLink_BytesRecv, STATGROUP_FaceAR);
+DECLARE_DWORD_COUNTER_STAT(TEXT("LiveLink Packets Recv"), STAT_FaceAR_LiveLink_PacketsRecv, STATGROUP_FaceAR);
 
 TSharedPtr<ILiveLinkSourceARKit> FAppleARKitLiveLinkSourceFactory::CreateLiveLinkSource(bool bCreateRemotePublisher)
 {
@@ -115,6 +128,8 @@ PRAGMA_DISABLE_DEPRECATION_WARNINGS
 
 void FAppleARKitLiveLinkSource::PublishBlendShapes(FName SubjectName, double Timestamp, uint32 FrameNumber, const FARBlendShapeMap& FaceBlendShapes)
 {
+	SCOPE_CYCLE_COUNTER(STAT_FaceAR_Local_PublishLiveLink);
+
 	check(Client != nullptr);
 	// This code touches UObjects so needs to be run only in the game thread
 	check(IsInGameThread());
@@ -258,6 +273,7 @@ void FAppleARKitLiveLinkRemotePublisher::PublishBlendShapes(FName SubjectName, d
 {
 	if (SendSocket != nullptr)
 	{
+		SCOPE_CYCLE_COUNTER(STAT_FaceAR_Remote_PublishLiveLink);
 		// Build the packet and send it
 		SendBuffer.Reset();
 		SendBuffer << BLEND_SHAPE_PACKET_VER;
@@ -283,6 +299,10 @@ void FAppleARKitLiveLinkRemotePublisher::PublishBlendShapes(FName SubjectName, d
 			ISocketSubsystem* SocketSub = ISocketSubsystem::Get();
 			UE_LOG(LogAppleARKitFace, Verbose, TEXT("Failed to send face AR packet with error (%s). Packet size (%d), sent (%d)"), SocketSub->GetSocketError(), SourceBufferSize, AmountSent);
 		}
+		INC_DWORD_STAT(STAT_FaceAR_LiveLink_PacketsSent);
+		INC_DWORD_STAT(STAT_FaceAR_Total_LiveLink_PacketsSent);
+		INC_DWORD_STAT_BY(STAT_FaceAR_LiveLink_BytesSent, AmountSent);
+		INC_DWORD_STAT_BY(STAT_FaceAR_Total_LiveLink_BytesSent, AmountSent);
 	}
 }
 
@@ -339,6 +359,8 @@ void FAppleARKitLiveLinkRemoteListener::InitLiveLinkSource()
 
 void FAppleARKitLiveLinkRemoteListener::Tick(float DeltaTime)
 {
+	SCOPE_CYCLE_COUNTER(STAT_FaceAR_ReceiveLiveLink);
+
 	uint32 BytesPending = 0;
 	while (RecvSocket->HasPendingData(BytesPending))
 	{
@@ -350,6 +372,11 @@ void FAppleARKitLiveLinkRemoteListener::Tick(float DeltaTime)
 			// Make sure the packet is a complete one and ignore if it is not
 			BytesRead > MIN_BLEND_SHAPE_PACKET_SIZE)
 		{
+			INC_DWORD_STAT(STAT_FaceAR_LiveLink_PacketsRecv);
+			INC_DWORD_STAT(STAT_FaceAR_Total_LiveLink_PacketsRecv);
+			INC_DWORD_STAT_BY(STAT_FaceAR_LiveLink_BytesRecv, BytesRead);
+			INC_DWORD_STAT_BY(STAT_FaceAR_Total_LiveLink_BytesRecv, BytesRead);
+
 			uint8 PacketVer = 0;
 			FName SubjectName;
 			double Timestamp = -1.0;
