Fixed deprecated code warning

diff --git a/Engine/Plugins/Runtime/AR/Apple/AppleARKitFaceSupport/Source/AppleARKitFaceSupport/Private/AppleARKitLiveLinkSource.cpp b/Engine/Plugins/Runtime/AR/Apple/AppleARKitFaceSupport/Source/AppleARKitFaceSupport/Private/AppleARKitLiveLinkSource.cpp
--- a/Engine/Plugins/Runtime/AR/Apple/AppleARKitFaceSupport/Source/AppleARKitFaceSupport/Private/AppleARKitLiveLinkSource.cpp
+++ b/Engine/Plugins/Runtime/AR/Apple/AppleARKitFaceSupport/Source/AppleARKitFaceSupport/Private/AppleARKitLiveLinkSource.cpp
@@ -364,7 +364,7 @@ FAppleARKitLiveLinkRemoteListener::~FAppleARKitLiveLinkRemoteListener()
 	{
 		RecvSocket->Close();
 
-		if (!GIsRequestingExit)
+		if (!IsEngineExitRequested())
 		{
 			ISocketSubsystem* SocketSub = ISocketSubsystem::Get();
 			SocketSub->DestroySocket(RecvSocket);
