Changed how face ar curves use timestamps when receiving cross platform so that there isn't a local mismatch that causes dropped frames

diff --git a/Engine/Plugins/Runtime/AppleARKitFaceSupport/Source/AppleARKitFaceSupport/Private/AppleARKitLiveLinkSource.cpp b/Engine/Plugins/Runtime/AppleARKitFaceSupport/Source/AppleARKitFaceSupport/Private/AppleARKitLiveLinkSource.cpp
--- a/Engine/Plugins/Runtime/AppleARKitFaceSupport/Source/AppleARKitFaceSupport/Private/AppleARKitLiveLinkSource.cpp
+++ b/Engine/Plugins/Runtime/AppleARKitFaceSupport/Source/AppleARKitFaceSupport/Private/AppleARKitLiveLinkSource.cpp
@@ -206,10 +206,13 @@ bool FAppleARKitLiveLinkSource::Exec(UWorld* InWorld, const TCHAR* Cmd, FOutputD
 	return false;
 }
 
-const uint8 BLEND_SHAPE_PACKET_VER = 2; // (ARKit 2.0)
+// 1 = Initial version
+// 2 = ARKit 2.0 extra blendshapes
+// 3 = Removed the timestamp to derive locally
+const uint8 BLEND_SHAPE_PACKET_VER = 3;
 
-const uint32 MAX_BLEND_SHAPE_PACKET_SIZE = sizeof(BLEND_SHAPE_PACKET_VER) + sizeof(double) + sizeof(uint32) + sizeof(uint8) + (sizeof(float) * (uint64)EARFaceBlendShape::MAX) + (sizeof(TCHAR) * 256);
-const uint32 MIN_BLEND_SHAPE_PACKET_SIZE = sizeof(BLEND_SHAPE_PACKET_VER) + sizeof(double) + sizeof(uint32) + sizeof(uint8) + (sizeof(float) * (uint64)EARFaceBlendShape::MAX) + sizeof(TCHAR);
+const uint32 MAX_BLEND_SHAPE_PACKET_SIZE = sizeof(BLEND_SHAPE_PACKET_VER) + sizeof(uint32) + sizeof(uint8) + (sizeof(float) * (uint64)EARFaceBlendShape::MAX) + (sizeof(TCHAR) * 256);
+const uint32 MIN_BLEND_SHAPE_PACKET_SIZE = sizeof(BLEND_SHAPE_PACKET_VER) + sizeof(uint32) + sizeof(uint8) + (sizeof(float) * (uint64)EARFaceBlendShape::MAX) + sizeof(TCHAR);
 
 FAppleARKitLiveLinkRemotePublisher::FAppleARKitLiveLinkRemotePublisher(const FString& InRemoteIp) :
 	RemoteIp(InRemoteIp),
@@ -278,7 +281,6 @@ void FAppleARKitLiveLinkRemotePublisher::PublishBlendShapes(FName SubjectName, d
 		SendBuffer.Reset();
 		SendBuffer << BLEND_SHAPE_PACKET_VER;
 		SendBuffer << SubjectName;
-		SendBuffer << Timestamp;
 		SendBuffer << FrameNumber;
 		uint8 BlendShapeCount = (uint8)EARFaceBlendShape::MAX;
 		check(FaceBlendShapes.Num() == BlendShapeCount);
@@ -379,7 +381,6 @@ void FAppleARKitLiveLinkRemoteListener::Tick(float DeltaTime)
 
 			uint8 PacketVer = 0;
 			FName SubjectName;
-			double Timestamp = -1.0;
 			uint32 FrameNumber = 0;
 			uint8 BlendShapeCount = (uint8)EARFaceBlendShape::MAX;
 
@@ -392,7 +393,6 @@ void FAppleARKitLiveLinkRemoteListener::Tick(float DeltaTime)
 				return;
 			}
 			FromBuffer >> SubjectName;
-			FromBuffer >> Timestamp;
 			FromBuffer >> FrameNumber;
 			FromBuffer >> BlendShapeCount;
 			if (FromBuffer.HasOverflow() || BlendShapeCount != (uint8)EARFaceBlendShape::MAX)
@@ -414,7 +414,7 @@ void FAppleARKitLiveLinkRemoteListener::Tick(float DeltaTime)
 				InitLiveLinkSource();
 				if (Source.IsValid())
 				{
-					Source->PublishBlendShapes(SubjectName, Timestamp, FrameNumber, BlendShapes);
+					Source->PublishBlendShapes(SubjectName, FPlatformTime::Seconds(), FrameNumber, BlendShapes);
 				}
 			}
 			else
