Added a device id in addition to subject name so that we can detect when a device change's its subject name

diff --git a/Engine/Plugins/Runtime/AppleARKitFaceSupport/Source/AppleARKitFaceSupport/Private/AppleARKitLiveLinkSource.cpp b/Engine/Plugins/Runtime/AppleARKitFaceSupport/Source/AppleARKitFaceSupport/Private/AppleARKitLiveLinkSource.cpp
--- a/Engine/Plugins/Runtime/AppleARKitFaceSupport/Source/AppleARKitFaceSupport/Private/AppleARKitLiveLinkSource.cpp
+++ b/Engine/Plugins/Runtime/AppleARKitFaceSupport/Source/AppleARKitFaceSupport/Private/AppleARKitLiveLinkSource.cpp
@@ -126,7 +126,7 @@ static FName ParseEnumName(FName EnumName)
 // Temporary for 4.20
 PRAGMA_DISABLE_DEPRECATION_WARNINGS
 
-void FAppleARKitLiveLinkSource::PublishBlendShapes(FName SubjectName, double Timestamp, uint32 FrameNumber, const FARBlendShapeMap& FaceBlendShapes)
+void FAppleARKitLiveLinkSource::PublishBlendShapes(FName SubjectName, double Timestamp, uint32 FrameNumber, const FARBlendShapeMap& FaceBlendShapes, FName DeviceId)
 {
 	SCOPE_CYCLE_COUNTER(STAT_FaceAR_Local_PublishLiveLink);
 
@@ -134,13 +134,28 @@ void FAppleARKitLiveLinkSource::PublishBlendShapes(FName SubjectName, double Tim
 	// This code touches UObjects so needs to be run only in the game thread
 	check(IsInGameThread());
 
-	if (SubjectName != LastSubjectName)
+	// Generate our device id if it wasn't passed in (local only case)
+	if (DeviceId == NAME_None)
 	{
-		Client->ClearSubject(LastSubjectName);
-		// We need to publish a skeleton for this subject name even though we doesn't use one
+		DeviceId = FName(*FPlatformMisc::GetDeviceId());
+	}
+
+	FName* LastSubjectNameForDeviceId = DeviceToLastSubjectNameMap.Find(DeviceId);
+	// Is this a new device and subject pair?
+	if (LastSubjectNameForDeviceId == nullptr)
+	{
+		// First time seen so publish an empty skeleton
+		Client->PushSubjectSkeleton(SourceGuid, SubjectName, FLiveLinkRefSkeleton());
+		DeviceToLastSubjectNameMap.Add(DeviceId, SubjectName);
+	}
+	// Did the subject name change for the device?
+	else if (SubjectName != *LastSubjectNameForDeviceId)
+	{
+		// The remote device changed subject names, so remove the old subject
+		Client->ClearSubject(*LastSubjectNameForDeviceId);
+		// Now add a new skeleton with the new subject name
 		Client->PushSubjectSkeleton(SourceGuid, SubjectName, FLiveLinkRefSkeleton());
 	}
-	LastSubjectName = SubjectName;
 
 	const UEnum* EnumPtr = FindObject<UEnum>(ANY_PACKAGE, TEXT("EARFaceBlendShape"), true);
 	if (EnumPtr != nullptr)
@@ -170,7 +185,7 @@ void FAppleARKitLiveLinkSource::PublishBlendShapes(FName SubjectName, double Tim
 		// Send it to the remote editor via the message bus
 		if (RemoteLiveLinkPublisher.IsValid())
 		{
-			RemoteLiveLinkPublisher->PublishBlendShapes(SubjectName, Timestamp, FrameNumber, FaceBlendShapes);
+			RemoteLiveLinkPublisher->PublishBlendShapes(SubjectName, Timestamp, FrameNumber, FaceBlendShapes, DeviceId);
 		}
 	}
 }
@@ -209,10 +224,11 @@ bool FAppleARKitLiveLinkSource::Exec(UWorld* InWorld, const TCHAR* Cmd, FOutputD
 // 1 = Initial version
 // 2 = ARKit 2.0 extra blendshapes
 // 3 = Removed the timestamp to derive locally
-const uint8 BLEND_SHAPE_PACKET_VER = 3;
+// 4 = Added the device id to stream so we can tell devices apart
+const uint8 BLEND_SHAPE_PACKET_VER = 4;
 
-const uint32 MAX_BLEND_SHAPE_PACKET_SIZE = sizeof(BLEND_SHAPE_PACKET_VER) + sizeof(uint32) + sizeof(uint8) + (sizeof(float) * (uint64)EARFaceBlendShape::MAX) + (sizeof(TCHAR) * 256);
-const uint32 MIN_BLEND_SHAPE_PACKET_SIZE = sizeof(BLEND_SHAPE_PACKET_VER) + sizeof(uint32) + sizeof(uint8) + (sizeof(float) * (uint64)EARFaceBlendShape::MAX) + sizeof(TCHAR);
+const uint32 MAX_BLEND_SHAPE_PACKET_SIZE = sizeof(BLEND_SHAPE_PACKET_VER) + sizeof(uint32) + sizeof(uint8) + (sizeof(float) * (uint64)EARFaceBlendShape::MAX) + (sizeof(TCHAR) * 256) + (sizeof(TCHAR) * 256);
+const uint32 MIN_BLEND_SHAPE_PACKET_SIZE = sizeof(BLEND_SHAPE_PACKET_VER) + sizeof(uint32) + sizeof(uint8) + (sizeof(float) * (uint64)EARFaceBlendShape::MAX) + sizeof(TCHAR) + sizeof(TCHAR);
 
 FAppleARKitLiveLinkRemotePublisher::FAppleARKitLiveLinkRemotePublisher(const FString& InRemoteIp) :
 	RemoteIp(InRemoteIp),
@@ -272,7 +288,7 @@ TSharedRef<FInternetAddr> FAppleARKitLiveLinkRemotePublisher::GetSendAddress()
 	return SendAddr;
 }
 
-void FAppleARKitLiveLinkRemotePublisher::PublishBlendShapes(FName SubjectName, double Timestamp, uint32 FrameNumber, const FARBlendShapeMap& FaceBlendShapes)
+void FAppleARKitLiveLinkRemotePublisher::PublishBlendShapes(FName SubjectName, double Timestamp, uint32 FrameNumber, const FARBlendShapeMap& FaceBlendShapes, FName DeviceId)
 {
 	if (SendSocket != nullptr)
 	{
@@ -280,6 +296,7 @@ void FAppleARKitLiveLinkRemotePublisher::PublishBlendShapes(FName SubjectName, d
 		// Build the packet and send it
 		SendBuffer.Reset();
 		SendBuffer << BLEND_SHAPE_PACKET_VER;
+		SendBuffer << DeviceId;
 		SendBuffer << SubjectName;
 		SendBuffer << FrameNumber;
 		uint8 BlendShapeCount = (uint8)EARFaceBlendShape::MAX;
@@ -383,6 +400,7 @@ void FAppleARKitLiveLinkRemoteListener::Tick(float DeltaTime)
 			FName SubjectName;
 			uint32 FrameNumber = 0;
 			uint8 BlendShapeCount = (uint8)EARFaceBlendShape::MAX;
+			FName DeviceId;
 
 			FNboSerializeFromBuffer FromBuffer(RecvBuffer.GetData(), BytesRead);
 
@@ -392,6 +410,7 @@ void FAppleARKitLiveLinkRemoteListener::Tick(float DeltaTime)
 				UE_LOG(LogAppleARKitFace, Verbose, TEXT("Packet overflow reading the packet version for the face AR packet"));
 				return;
 			}
+			FromBuffer >> DeviceId;
 			FromBuffer >> SubjectName;
 			FromBuffer >> FrameNumber;
 			FromBuffer >> BlendShapeCount;
@@ -414,7 +433,7 @@ void FAppleARKitLiveLinkRemoteListener::Tick(float DeltaTime)
 				InitLiveLinkSource();
 				if (Source.IsValid())
 				{
-					Source->PublishBlendShapes(SubjectName, FPlatformTime::Seconds(), FrameNumber, BlendShapes);
+					Source->PublishBlendShapes(SubjectName, FPlatformTime::Seconds(), FrameNumber, BlendShapes, DeviceId);
 				}
 			}
 			else
diff --git a/Engine/Plugins/Runtime/AppleARKitFaceSupport/Source/AppleARKitFaceSupport/Private/AppleARKitLiveLinkSource.h b/Engine/Plugins/Runtime/AppleARKitFaceSupport/Source/AppleARKitFaceSupport/Private/AppleARKitLiveLinkSource.h
--- a/Engine/Plugins/Runtime/AppleARKitFaceSupport/Source/AppleARKitFaceSupport/Private/AppleARKitLiveLinkSource.h
+++ b/Engine/Plugins/Runtime/AppleARKitFaceSupport/Source/AppleARKitFaceSupport/Private/AppleARKitLiveLinkSource.h
@@ -733,7 +733,7 @@ public:
 	virtual ~FAppleARKitLiveLinkRemotePublisher();
 
 	// IARKitBlendShapePublisher interface
-	virtual void PublishBlendShapes(FName SubjectName, double Timestamp, uint32 FrameNumber, const FARBlendShapeMap& FaceBlendShapes) override;
+	virtual void PublishBlendShapes(FName SubjectName, double Timestamp, uint32 FrameNumber, const FARBlendShapeMap& FaceBlendShapes, FName DeviceId = NAME_None) override;
 	// End IARKitBlendShapePublisher
 
 	bool InitSendSocket();
@@ -769,7 +769,7 @@ private:
 	// End ILiveLinkSource
 
 	// IARKitBlendShapePublisher interface
-	virtual void PublishBlendShapes(FName SubjectName, double Timestamp, uint32 FrameNumber, const FARBlendShapeMap& FaceBlendShapes) override;
+	virtual void PublishBlendShapes(FName SubjectName, double Timestamp, uint32 FrameNumber, const FARBlendShapeMap& FaceBlendShapes, FName DeviceId = NAME_None) override;
 	// End IARKitBlendShapePublisher
 
 	//~ FSelfRegisteringExec
@@ -788,6 +788,6 @@ private:
 	/** The last time we sent the data. Used to not send redundant data */
 	uint32 LastFramePublished;
 
-	/** Used to track names changes */
-	FName LastSubjectName;
+	/** Used to track names changes for a given device */
+	TMap<FName, FName> DeviceToLastSubjectNameMap;
 };
diff --git a/Engine/Plugins/Runtime/AppleARKitFaceSupport/Source/AppleARKitFaceSupport/Public/AppleARKitLiveLinkSourceFactory.h b/Engine/Plugins/Runtime/AppleARKitFaceSupport/Source/AppleARKitFaceSupport/Public/AppleARKitLiveLinkSourceFactory.h
--- a/Engine/Plugins/Runtime/AppleARKitFaceSupport/Source/AppleARKitFaceSupport/Public/AppleARKitLiveLinkSourceFactory.h
+++ b/Engine/Plugins/Runtime/AppleARKitFaceSupport/Source/AppleARKitFaceSupport/Public/AppleARKitLiveLinkSourceFactory.h
@@ -11,7 +11,7 @@
 class APPLEARKITFACESUPPORT_API IARKitBlendShapePublisher
 {
 public:
-	virtual void PublishBlendShapes(FName SubjectName, double Timestamp, uint32 FrameNumber, const FARBlendShapeMap& BlendShapes) = 0;
+	virtual void PublishBlendShapes(FName SubjectName, double Timestamp, uint32 FrameNumber, const FARBlendShapeMap& BlendShapes, FName DeviceID = NAME_None) = 0;
 };
 
 /** Interface that publishes face ar blend shape information via LiveLink */
