First pass Face AR networking support (lightly tested, caveat emptor)

diff --git a/Engine/Plugins/Runtime/AppleARKitFaceSupport/Source/AppleARKitFaceSupport/Private/AppleARKitFaceMeshComponent.cpp b/Engine/Plugins/Runtime/AppleARKitFaceSupport/Source/AppleARKitFaceSupport/Private/AppleARKitFaceMeshComponent.cpp
--- a/Engine/Plugins/Runtime/AppleARKitFaceSupport/Source/AppleARKitFaceSupport/Private/AppleARKitFaceMeshComponent.cpp
+++ b/Engine/Plugins/Runtime/AppleARKitFaceSupport/Source/AppleARKitFaceSupport/Private/AppleARKitFaceMeshComponent.cpp
@@ -13,6 +13,7 @@
 #include "ARSessionConfig.h"
 #include "AppleARKitSettings.h"
 #include "AppleARKitConversion.h"
+#include "Net/UnrealNetwork.h"
 
 DECLARE_CYCLE_STAT(TEXT("Update"), STAT_FaceAR_Component_Update, STATGROUP_FaceAR);
 
@@ -173,18 +174,26 @@ void UAppleARKitFaceMeshComponent::CreateMesh(const TArray<FVector>& Vertices, c
 
 void UAppleARKitFaceMeshComponent::SetBlendShapes(const TMap<EARFaceBlendShape, float>& InBlendShapes)
 {
-	BlendShapes = InBlendShapes;
 	LastUpdateFrameNumber++;
 	LastUpdateTimestamp = FPlatformTime::Seconds();
-	if (LiveLinkSource.IsValid())
-	{
-		LiveLinkSource->PublishBlendShapes(LiveLinkSubjectName, LastUpdateTimestamp, LastUpdateFrameNumber, BlendShapes);
-	}
+	BuildUpdatedCurves(InBlendShapes);
 }
 
 void UAppleARKitFaceMeshComponent::SetBlendShapeAmount(EARFaceBlendShape BlendShape, float Amount)
 {
-	BlendShapes.Add(BlendShape, Amount);
+	RemoteCurves.Reset();
+	float OldVal = BlendShapes[BlendShape];
+	if (FNetQuantizeFaceCurve::IsDifferentEnough(OldVal, Amount))
+	{
+		new(RemoteCurves) FNetQuantizeFaceCurve(BlendShape, Amount);
+	}
+	if (GetNetMode() == NM_Client)
+	{
+		// Send to the server for replication
+		ServerUpdateFaceCurves(RemoteCurves);
+	}
+	// Merge them in using the OnRep
+	OnRep_RemoteCurves();
 }
 
 float UAppleARKitFaceMeshComponent::GetFaceBlendShapeAmount(EARFaceBlendShape BlendShape) const
@@ -337,9 +346,11 @@ void UAppleARKitFaceMeshComponent::TickComponent(float DeltaTime, ELevelTick, FA
 			TrackedRot.Roll = -TrackedRot.Roll;
 			LocalToWorldTransform.SetRotation(FQuat(TrackedRot));
 		}
-		BlendShapes = FaceGeometry->GetBlendShapes();
+
 		LastUpdateFrameNumber = FaceGeometry->GetLastUpdateFrameNumber();
 		LastUpdateTimestamp = FaceGeometry->GetLastUpdateTimestamp();
+		BuildUpdatedCurves(FaceGeometry->GetBlendShapes());
+
 		// Create or update the mesh depending on if we've been created before
 		if (GetNumSections() > 0)
 		{
@@ -365,3 +376,65 @@ FTransform UAppleARKitFaceMeshComponent::GetTransform() const
 {
 	return LocalToWorldTransform;
 }
+
+void UAppleARKitFaceMeshComponent::OnRep_RemoteCurves()
+{
+	if (RemoteCurves.Num() > 0)
+	{
+		// Merge the values into our map
+		for (const FNetQuantizeFaceCurve& Curve : RemoteCurves)
+		{
+			BlendShapes.Add(Curve.GetBlendShape(), Curve.GetAmountAsFloat());
+		}
+		// We have to manage this since the ar system isn't updating this locally for us
+		LastUpdateFrameNumber++;
+		LastUpdateTimestamp = FPlatformTime::Seconds();
+		// If we are publishing to LiveLink, then push the update out
+		if (LiveLinkSource.IsValid())
+		{
+			LiveLinkSource->PublishBlendShapes(LiveLinkSubjectName, LastUpdateTimestamp, LastUpdateFrameNumber, BlendShapes);
+		}
+	}
+}
+
+void UAppleARKitFaceMeshComponent::GetLifetimeReplicatedProps(TArray<FLifetimeProperty>& OutLifetimeProps) const
+{
+	Super::GetLifetimeReplicatedProps(OutLifetimeProps);
+	
+	DOREPLIFETIME(UAppleARKitFaceMeshComponent, RemoteCurves);
+}
+
+bool UAppleARKitFaceMeshComponent::ServerUpdateFaceCurves_Validate(const TArray<FNetQuantizeFaceCurve>& ClientCurves)
+{
+	return ClientCurves.Num() < (int32)EARFaceBlendShape::MAX;
+}
+
+void UAppleARKitFaceMeshComponent::ServerUpdateFaceCurves_Implementation(const TArray<FNetQuantizeFaceCurve>& ClientCurves)
+{
+	RemoteCurves = ClientCurves;
+	// Publish these locally for the server
+	OnRep_RemoteCurves();
+}
+
+void UAppleARKitFaceMeshComponent::BuildUpdatedCurves(const FARBlendShapeMap& NewCurves)
+{
+	RemoteCurves.Reset();
+	// Build the set of deltas
+	for (uint8 Shape = 0; Shape < (uint8)EARFaceBlendShape::MAX; Shape++)
+	{
+		EARFaceBlendShape BlendShape = (EARFaceBlendShape)Shape;
+		float OldVal = BlendShapes[BlendShape];
+		float NewVal = NewCurves[BlendShape];
+		if (FNetQuantizeFaceCurve::IsDifferentEnough(OldVal, NewVal))
+		{
+			new(RemoteCurves) FNetQuantizeFaceCurve(BlendShape, NewVal);
+		}
+	}
+	if (GetNetMode() == NM_Client)
+	{
+		// Send to the server for replication
+		ServerUpdateFaceCurves(RemoteCurves);
+	}
+	// Merge them in using the OnRep
+	OnRep_RemoteCurves();
+}
diff --git a/Engine/Plugins/Runtime/AppleARKitFaceSupport/Source/AppleARKitFaceSupport/Public/AppleARKitFaceMeshComponent.h b/Engine/Plugins/Runtime/AppleARKitFaceSupport/Source/AppleARKitFaceSupport/Public/AppleARKitFaceMeshComponent.h
--- a/Engine/Plugins/Runtime/AppleARKitFaceSupport/Source/AppleARKitFaceSupport/Public/AppleARKitFaceMeshComponent.h
+++ b/Engine/Plugins/Runtime/AppleARKitFaceSupport/Source/AppleARKitFaceSupport/Public/AppleARKitFaceMeshComponent.h
@@ -23,6 +23,92 @@ enum class EARFaceComponentTransformMixing : uint8
 	TrackingOnly
 };
 
+/**
+ * Packs the curve into 2 bytes with the amount being +/- 127
+ */
+USTRUCT()
+struct FNetQuantizeFaceCurve
+{
+	GENERATED_USTRUCT_BODY()
+	
+	FORCEINLINE FNetQuantizeFaceCurve()
+	{
+	}
+	
+	FORCEINLINE FNetQuantizeFaceCurve(EARFaceBlendShape InBlendShape, float InAmount)
+		: BlendShape((uint8)InBlendShape)
+		, Amount(ConvertAmountToInt(InAmount))
+	{
+	}
+	
+	FORCEINLINE FNetQuantizeFaceCurve(const FNetQuantizeFaceCurve& Other)
+		: BlendShape(Other.BlendShape)
+		, Amount(Other.Amount)
+	{
+	}
+	
+	FORCEINLINE float GetAmountAsFloat() const
+	{
+		return ConvertAmountToFloat(Amount);
+	}
+	
+	FORCEINLINE EARFaceBlendShape GetBlendShape() const
+	{
+		check(BlendShape < (uint8)EARFaceBlendShape::MAX);
+		return (EARFaceBlendShape)BlendShape;
+	}
+
+	bool NetSerialize(FArchive& Ar, class UPackageMap* Map, bool& bOutSuccess)
+	{
+		uint16 Value = 0;
+		if (Ar.IsSaving())
+		{
+			Value = ((uint8)Amount | (BlendShape << 8));
+			Ar << Value;
+		}
+		else
+		{
+			Ar << Value;
+			Amount = (int8)(Value & 0xFF);
+			BlendShape = ((Value >> 8) & 0xFF);
+		}
+		bOutSuccess = true;
+		return true;
+	}
+	
+	static bool IsDifferentEnough(float Val1, float Val2)
+	{
+		return ConvertAmountToInt(Val1) != ConvertAmountToInt(Val2);
+	}
+
+private:
+	static constexpr float Scale = 127.f;
+	static constexpr float InvScale = 1.f / Scale;
+	
+	static FORCEINLINE float ConvertAmountToFloat(int8 InAmount)
+	{
+		return ((float)InAmount) * InvScale;
+	}
+	
+	static FORCEINLINE int8 ConvertAmountToInt(float InAmount)
+	{
+		return (int8)FMath::TruncToInt(InAmount * Scale);
+	}
+	
+	uint8 BlendShape;
+	int8 Amount;
+};
+
+template<>
+struct TStructOpsTypeTraits< FNetQuantizeFaceCurve > : public TStructOpsTypeTraitsBase2< FNetQuantizeFaceCurve >
+{
+	enum
+	{
+		WithNetSerializer = true,
+		WithNetSharedSerialization = true,
+	};
+};
+
 /**
  * This component is updated by the ARSystem with face data on devices that have support for it
  */
@@ -108,6 +194,10 @@ public:
 	UFUNCTION(BlueprintPure, Category = "Components|ARFaceMesh")
 	FTransform GetTransform() const;
 
+	/**	Indicates whether the face mesh data should be built for rendering or not */
+	UPROPERTY(EditAnywhere, BlueprintReadOnly, Category = "Components|ARFaceMesh")
+	bool bWantsMeshUpdates;
+
 	/**	Indicates whether collision should be created for this face mesh. This adds significant cost, so only use if you need to trace against the face mesh. */
 	UPROPERTY(EditAnywhere, BlueprintReadOnly, Category = "Components|ARFaceMesh")
 	bool bWantsCollision;
@@ -132,6 +222,10 @@ public:
 	UPROPERTY(EditAnywhere, BlueprintReadOnly, Category = "Components|ARFaceMesh")
 	FName LiveLinkSubjectName;
 
+	/** The set of changed curves to replicate to the other clients */
+	UPROPERTY(ReplicatedUsing=OnRep_RemoteCurves, Transient)
+	TArray<FNetQuantizeFaceCurve> RemoteCurves;
+
 private:
 	virtual void TickComponent(float DeltaTime, enum ELevelTick TickType, FActorComponentTickFunction* ThisTickFunction) override;
 	virtual void InitializeComponent() override;
@@ -139,8 +233,27 @@ private:
 	virtual class UMaterialInterface* GetMaterial(int32 ElementIndex) const override;
 	virtual void SetMaterial(int32 ElementIndex, class UMaterialInterface* Material) override;
 
-	UARFaceGeometry* FindFaceGeometry();
+	/** Merges in the face curve deltas and pushes them to LiveLink */
+	UFUNCTION()
+	void OnRep_RemoteCurves();
+
+	/**
+	 * Sends the updated curves from the client to the server so that it can replicate to other clients
+	 *
+	 * @param ClientCurves the client's set of updated curves
+	 */
+	UFUNCTION(Reliable, Server, WithValidation)
+	void ServerUpdateFaceCurves(const TArray<FNetQuantizeFaceCurve>& ClientCurves);
+
+	/**
+	 * Builds the delta set of curves needed for replication
+	 *
+	 * @param NewCurves the updated set of curves from the ar system
+	 */
+	void BuildUpdatedCurves(const FARBlendShapeMap& NewCurves);
 
+	UARFaceGeometry* FindFaceGeometry();
+	
 	/** The current set of blend shapes for this component instance */
 	FARBlendShapeMap BlendShapes;
 	/** Transform of the face mesh */
