ARKit plugins are no longer experimental

diff --git a/Engine/Plugins/Runtime/AR/Apple/AppleARKit/AppleARKit.uplugin b/Engine/Plugins/Runtime/AR/Apple/AppleARKit/AppleARKit.uplugin
--- a/Engine/Plugins/Runtime/AR/Apple/AppleARKit/AppleARKit.uplugin
+++ b/Engine/Plugins/Runtime/AR/Apple/AppleARKit/AppleARKit.uplugin
@@ -11,7 +11,7 @@
 	"MarketplaceURL": "",
 	"SupportURL": "",
 	"CanContainContent": true,
-	"IsBetaVersion": true,
+	"IsBetaVersion": false,
 	"EnabledByDefault": false,
 	"Modules": [
 		{
diff --git a/Engine/Plugins/Runtime/AR/Apple/AppleARKitFaceSupport/AppleARKitFaceSupport.uplugin b/Engine/Plugins/Runtime/AR/Apple/AppleARKitFaceSupport/AppleARKitFaceSupport.uplugin
--- a/Engine/Plugins/Runtime/AR/Apple/AppleARKitFaceSupport/AppleARKitFaceSupport.uplugin
+++ b/Engine/Plugins/Runtime/AR/Apple/AppleARKitFaceSupport/AppleARKitFaceSupport.uplugin
@@ -11,7 +11,7 @@
 	"MarketplaceURL": "",
 	"SupportURL": "",
 	"CanContainContent": true,
-	"IsBetaVersion": true,
+	"IsBetaVersion": false,
 	"EnabledByDefault": false,
 	"Modules":
 	[
diff --git a/Engine/Plugins/Runtime/AR/Apple/AppleARKitFaceSupport/Source/AppleARKitFaceSupport/AppleARKitFaceSupport.Build.cs b/Engine/Plugins/Runtime/AR/Apple/AppleARKitFaceSupport/Source/AppleARKitFaceSupport/AppleARKitFaceSupport.Build.cs
--- a/Engine/Plugins/Runtime/AR/Apple/AppleARKitFaceSupport/Source/AppleARKitFaceSupport/AppleARKitFaceSupport.Build.cs
+++ b/Engine/Plugins/Runtime/AR/Apple/AppleARKitFaceSupport/Source/AppleARKitFaceSupport/AppleARKitFaceSupport.Build.cs
@@ -8,7 +8,7 @@ public class AppleARKitFaceSupport : ModuleRules
 	{
 		PrivateIncludePaths.AddRange(
 			new string[] {
-				"../../../../../../Plugins/Experimental/AR/Apple/AppleARKit/Source/AppleARKit/Private",
+				"../../../../../../Plugins/Runtime/AR/Apple/AppleARKit/Source/AppleARKit/Private",
 				// ... add other private include paths required here ...
 			}
 			);
