Moved where remote sending of face tracking data is done
- Sending of the blendshapes now happens on a background thread scheduled immediately when ARKit notifies of change
- Framework for async file writing of the curve was stubbed in (WIP)

diff --git a/Engine/Plugins/Runtime/AR/Apple/AppleARKit/Source/AppleARKit/Public/AppleARKitConversion.h b/Engine/Plugins/Runtime/AR/Apple/AppleARKit/Source/AppleARKit/Public/AppleARKitConversion.h
--- a/Engine/Plugins/Runtime/AR/Apple/AppleARKit/Source/AppleARKit/Public/AppleARKitConversion.h
+++ b/Engine/Plugins/Runtime/AR/Apple/AppleARKit/Source/AppleARKit/Public/AppleARKitConversion.h
@@ -256,6 +256,21 @@ enum class EAppleAnchorType : uint8
 
 struct FAppleARKitAnchorData
 {
+	FAppleARKitAnchorData()
+		: AnchorType()
+		, AnchorGUID(FGuid())
+		, ProbeTexture(nullptr)
+		, Timestamp(0.0)
+		, FrameNumber(0)
+		, bIsTracked(false)
+	{
+	}
+
+	FAppleARKitAnchorData(const FAppleARKitAnchorData& Other)
+	{
+		Copy(Other);
+	}
+
 	FAppleARKitAnchorData(FGuid InAnchorGuid, FTransform InTransform)
 		: Transform( InTransform )
 		, AnchorType( EAppleAnchorType::Anchor )
@@ -304,6 +319,41 @@ struct FAppleARKitAnchorData
 	{
 	}
 
+	FAppleARKitAnchorData& operator=(const FAppleARKitAnchorData& Other)
+	{
+		if (this != &Other)
+		{
+			Copy(Other);
+		}
+
+		return *this;
+	}
+
+	void Copy(const FAppleARKitAnchorData& Other)
+	{
+		Transform = Other.Transform;
+		AnchorType = Other.AnchorType;
+		AnchorGUID = Other.AnchorGUID;
+		Center = Other.Center;
+		Extent = Other.Center;
+		BoundaryVerts = Other.BoundaryVerts;
+
+		BlendShapes = Other.BlendShapes;
+		FaceVerts = Other.FaceVerts;
+
+		DetectedAnchorName = Other.DetectedAnchorName;
+
+		ProbeTexture = Other.ProbeTexture;
+
+		LeftEyeTransform = Other.LeftEyeTransform;
+		RightEyeTransform = Other.RightEyeTransform;
+		LookAtTarget = Other.LookAtTarget;
+		Timestamp = Other.Timestamp;
+		FrameNumber = Other.FrameNumber;
+
+		bIsTracked = Other.bIsTracked;
+	}
+
 	FTransform Transform;
 	EAppleAnchorType AnchorType;
 	FGuid AnchorGUID;
diff --git a/Engine/Plugins/Runtime/AR/Apple/AppleARKitFaceSupport/Source/AppleARKitFaceSupport/Private/AppleARKitFaceMeshComponent.cpp b/Engine/Plugins/Runtime/AR/Apple/AppleARKitFaceSupport/Source/AppleARKitFaceSupport/Private/AppleARKitFaceMeshComponent.cpp
--- a/Engine/Plugins/Runtime/AR/Apple/AppleARKitFaceSupport/Source/AppleARKitFaceSupport/Private/AppleARKitFaceMeshComponent.cpp
+++ b/Engine/Plugins/Runtime/AR/Apple/AppleARKitFaceSupport/Source/AppleARKitFaceSupport/Private/AppleARKitFaceMeshComponent.cpp
@@ -368,7 +368,7 @@ void UAppleARKitFaceMeshComponent::PublishViaLiveLink(FName SubjectName)
 	LiveLinkSubjectName = SubjectName;
 	if (!LiveLinkSource.IsValid())
 	{
-		LiveLinkSource = FAppleARKitLiveLinkSourceFactory::CreateLiveLinkSource(false);
+		LiveLinkSource = FAppleARKitLiveLinkSourceFactory::CreateLiveLinkSource();
 	}
 }
 
diff --git a/Engine/Plugins/Runtime/AR/Apple/AppleARKitFaceSupport/Source/AppleARKitFaceSupport/Private/AppleARKitFaceSupportImpl.cpp b/Engine/Plugins/Runtime/AR/Apple/AppleARKitFaceSupport/Source/AppleARKitFaceSupport/Private/AppleARKitFaceSupportImpl.cpp
--- a/Engine/Plugins/Runtime/AR/Apple/AppleARKitFaceSupport/Source/AppleARKitFaceSupport/Private/AppleARKitFaceSupportImpl.cpp
+++ b/Engine/Plugins/Runtime/AR/Apple/AppleARKitFaceSupport/Source/AppleARKitFaceSupport/Private/AppleARKitFaceSupportImpl.cpp
@@ -8,6 +8,7 @@
 #include "ARSystem.h"
 #include "Misc/ConfigCacheIni.h"
 #include "AppleARKitFaceSupportModule.h"
+#include "Async/Async.h"
 
 DECLARE_CYCLE_STAT(TEXT("Conversion"), STAT_FaceAR_Conversion, STATGROUP_FaceAR);
 
@@ -80,6 +81,26 @@ void FAppleARKitFaceSupport::Shutdown()
 	IModularFeatures::Get().UnregisterModularFeature(GetModularFeatureName(), this);
 }
 
+bool FAppleARKitFaceSupport::Exec(UWorld*, const TCHAR* Cmd, FOutputDevice& Ar)
+{
+	if (FParse::Command(&Cmd, TEXT("LiveLinkFaceAR")))
+	{
+		FString RemoteIp;
+		if (FParse::Value(Cmd, TEXT("SendTo="), RemoteIp))
+		{
+			// We need to recreate the LiveLink remote publisher
+			RemoteLiveLinkPublisher = nullptr;
+			// Only send from iOS to desktop
+#if PLATFORM_IOS
+			// This will perform the sending of the data to the remote
+			RemoteLiveLinkPublisher = FAppleARKitLiveLinkSourceFactory::CreateLiveLinkRemotePublisher(RemoteIp);
+#endif
+			return true;
+		}
+	}
+	return false;
+}
+
 #if SUPPORTS_ARKIT_1_0
 
 ARConfiguration* FAppleARKitFaceSupport::ToARConfiguration(UARSessionConfig* SessionConfig)
@@ -122,12 +143,36 @@ TArray<TSharedPtr<FAppleARKitAnchorData>> FAppleARKitFaceSupport::MakeAnchorData
 		if (AnchorData.IsValid())
 		{
 			AnchorList.Add(AnchorData);
+			// Process any providers that want real time access to the face curve data
+			ProcessRealTimePublishers(AnchorData);
 		}
 	}
 
 	return AnchorList;
 }
 
+void FAppleARKitFaceSupport::ProcessRealTimePublishers(TSharedPtr<FAppleARKitAnchorData> AnchorData)
+{
+	// Copy the data from the passed in anchor
+	AsyncAnchorCopy = *AnchorData;
+	AsyncTask(ENamedThreads::AnyBackgroundThreadNormalTask, [this]()
+	{
+		double Timestamp = AsyncAnchorCopy.Timestamp;
+		uint32 FrameNumber = AsyncAnchorCopy.FrameNumber;
+		const FARBlendShapeMap& BlendShapes = AsyncAnchorCopy.BlendShapes;
+
+		if (RemoteLiveLinkPublisher.IsValid())
+		{
+			RemoteLiveLinkPublisher->PublishBlendShapes(FaceTrackingLiveLinkSubjectName, Timestamp, FrameNumber, BlendShapes, LocalDeviceId);
+		}
+
+		if (LiveLinkFileWriter.IsValid())
+		{
+			LiveLinkFileWriter->PublishBlendShapes(FaceTrackingLiveLinkSubjectName, Timestamp, FrameNumber, BlendShapes, LocalDeviceId);
+		}
+	});
+}
+
 void FAppleARKitFaceSupport::PublishLiveLinkData(TSharedPtr<FAppleARKitAnchorData> Anchor)
 {
 	static bool bNeedsInit = true;
@@ -139,7 +184,7 @@ void FAppleARKitFaceSupport::PublishLiveLinkData(TSharedPtr<FAppleARKitAnchorDat
 		{
 			FaceTrackingLiveLinkSubjectName = GetDefault<UAppleARKitSettings>()->DefaultFaceTrackingLiveLinkSubjectName;
 	#if PLATFORM_IOS
-			LiveLinkSource = FAppleARKitLiveLinkSourceFactory::CreateLiveLinkSource(true);
+			LiveLinkSource = FAppleARKitLiveLinkSourceFactory::CreateLiveLinkSource();
 	#else
 			// This should be started already, but just in case
 			FAppleARKitLiveLinkSourceFactory::CreateLiveLinkRemoteListener();
diff --git a/Engine/Plugins/Runtime/AR/Apple/AppleARKitFaceSupport/Source/AppleARKitFaceSupport/Private/AppleARKitFaceSupportImpl.h b/Engine/Plugins/Runtime/AR/Apple/AppleARKitFaceSupport/Source/AppleARKitFaceSupport/Private/AppleARKitFaceSupportImpl.h
--- a/Engine/Plugins/Runtime/AR/Apple/AppleARKitFaceSupport/Source/AppleARKitFaceSupport/Private/AppleARKitFaceSupportImpl.h
+++ b/Engine/Plugins/Runtime/AR/Apple/AppleARKitFaceSupport/Source/AppleARKitFaceSupport/Private/AppleARKitFaceSupportImpl.h
@@ -10,6 +10,7 @@
 
 class APPLEARKITFACESUPPORT_API FAppleARKitFaceSupport :
 	public IAppleARKitFaceSupport,
+	public FSelfRegisteringExec,
 	public TSharedFromThis<FAppleARKitFaceSupport, ESPMode::ThreadSafe>
 {
 public:
@@ -30,13 +31,28 @@ private:
 	virtual TArray<FARVideoFormat> ToARConfiguration() override;
 #endif
 	// ~IAppleARKitFaceSupport
+
+	/** Publishes the remote publisher and the file writer if present */
+	void ProcessRealTimePublishers(TSharedPtr<FAppleARKitAnchorData> AnchorData);
 #endif
 
+	//~ FSelfRegisteringExec
+	virtual bool Exec(UWorld*, const TCHAR* Cmd, FOutputDevice& Ar) override;
+	//~ FSelfRegisteringExec
+
 	/** If requested, publishes face ar updates to LiveLink for the animation system to use */
 	TSharedPtr<ILiveLinkSourceARKit> LiveLinkSource;
 	/** Copied from the UARSessionConfig project settings object */
 	FName FaceTrackingLiveLinkSubjectName;
 	/** The id of this device */
 	FName LocalDeviceId;
+	/** A publisher that sends to a remote machine */
+	TSharedPtr<IARKitBlendShapePublisher, ESPMode::ThreadSafe> RemoteLiveLinkPublisher;
+	/** A publisher that writes the data to disk */
+	TSharedPtr<IARKitBlendShapePublisher, ESPMode::ThreadSafe> LiveLinkFileWriter;
+#if SUPPORTS_ARKIT_1_0
+	/** Used to hold a copy of the data to be processed by the background thread */
+	FAppleARKitAnchorData AsyncAnchorCopy;
+#endif
 };
 
diff --git a/Engine/Plugins/Runtime/AR/Apple/AppleARKitFaceSupport/Source/AppleARKitFaceSupport/Private/AppleARKitLiveLinkSource.cpp b/Engine/Plugins/Runtime/AR/Apple/AppleARKitFaceSupport/Source/AppleARKitFaceSupport/Private/AppleARKitLiveLinkSource.cpp
--- a/Engine/Plugins/Runtime/AR/Apple/AppleARKitFaceSupport/Source/AppleARKitFaceSupport/Private/AppleARKitLiveLinkSource.cpp
+++ b/Engine/Plugins/Runtime/AR/Apple/AppleARKitFaceSupport/Source/AppleARKitFaceSupport/Private/AppleARKitLiveLinkSource.cpp
@@ -33,14 +33,14 @@ DECLARE_DWORD_COUNTER_STAT(TEXT("LiveLink Packets Sent"), STAT_FaceAR_LiveLink_P
 DECLARE_DWORD_COUNTER_STAT(TEXT("LiveLink Bytes Recv"), STAT_FaceAR_LiveLink_BytesRecv, STATGROUP_FaceAR);
 DECLARE_DWORD_COUNTER_STAT(TEXT("LiveLink Packets Recv"), STAT_FaceAR_LiveLink_PacketsRecv, STATGROUP_FaceAR);
 
-TSharedPtr<ILiveLinkSourceARKit> FAppleARKitLiveLinkSourceFactory::CreateLiveLinkSource(bool bCreateRemotePublisher)
+TSharedPtr<ILiveLinkSourceARKit> FAppleARKitLiveLinkSourceFactory::CreateLiveLinkSource()
 {
 	IModularFeatures& ModularFeatures = IModularFeatures::Get();
 
 	if (ModularFeatures.IsModularFeatureAvailable(ILiveLinkClient::ModularFeatureName))
 	{
 		ILiveLinkClient* LiveLinkClient = &IModularFeatures::Get().GetModularFeature<ILiveLinkClient>(ILiveLinkClient::ModularFeatureName);
-		TSharedPtr <ILiveLinkSourceARKit> Source = MakeShareable(new FAppleARKitLiveLinkSource(bCreateRemotePublisher));
+		TSharedPtr <ILiveLinkSourceARKit> Source = MakeShareable(new FAppleARKitLiveLinkSource());
 		LiveLinkClient->AddSource(Source);
 		return Source;
 	}
@@ -61,27 +61,36 @@ void FAppleARKitLiveLinkSourceFactory::CreateLiveLinkRemoteListener()
 	}
 }
 
-FAppleARKitLiveLinkSource::FAppleARKitLiveLinkSource(bool bCreateRemotePublisher) :
-	Client(nullptr)
-	, LastFramePublished(0)
+TSharedPtr<IARKitBlendShapePublisher, ESPMode::ThreadSafe> FAppleARKitLiveLinkSourceFactory::CreateLiveLinkRemotePublisher(const FString& RemoteAddr)
 {
+	TSharedPtr<IARKitBlendShapePublisher, ESPMode::ThreadSafe> RemoteLiveLinkPublisher;
+	// Only send from iOS to desktop
 #if PLATFORM_IOS
-	if (bCreateRemotePublisher)
+	// This will perform the sending of the data to the remote
+	FAppleARKitLiveLinkRemotePublisher* Publisher = new FAppleARKitLiveLinkRemotePublisher(RemoteAddr);
+	if (Publisher->InitSendSocket())
 	{
-		// Only send from iOS to desktop
-		// This will perform the sending of the data to the remote
-		FAppleARKitLiveLinkRemotePublisher* Publisher = new FAppleARKitLiveLinkRemotePublisher();
-		if (Publisher->InitSendSocket())
-		{
-			RemoteLiveLinkPublisher = MakeShareable(Publisher);
-		}
-		else
-		{
-			UE_LOG(LogAppleARKitFace, Warning, TEXT("Failed to create LiveLink remote publisher, so no data will be sent out"));
-			delete Publisher;
-		}
+		RemoteLiveLinkPublisher = MakeShareable(Publisher);
+	}
+	else
+	{
+		UE_LOG(LogAppleARKitFace, Warning, TEXT("Failed to create LiveLink remote publisher, so no data will be sent out"));
+		delete Publisher;
 	}
 #endif
+	return RemoteLiveLinkPublisher;
+}
+
+TSharedPtr<IARKitBlendShapePublisher, ESPMode::ThreadSafe> FAppleARKitLiveLinkSourceFactory::CreateLiveLinkLocalFileWriter()
+{
+	TSharedPtr<IARKitBlendShapePublisher, ESPMode::ThreadSafe> LocalFileWriter;
+	return LocalFileWriter;
+}
+
+FAppleARKitLiveLinkSource::FAppleARKitLiveLinkSource() :
+	Client(nullptr)
+	, LastFramePublished(0)
+{
 }
 
 void FAppleARKitLiveLinkSource::ReceiveClient(ILiveLinkClient* InClient, FGuid InSourceGuid)
@@ -176,45 +185,11 @@ void FAppleARKitLiveLinkSource::PublishBlendShapes(FName SubjectName, double Tim
 
 		// Share the data locally with the LiveLink client
 		Client->PushSubjectData(SourceGuid, SubjectName, LiveLinkFrame);
-		// Send it to the remote editor via the message bus
-		if (RemoteLiveLinkPublisher.IsValid())
-		{
-			RemoteLiveLinkPublisher->PublishBlendShapes(SubjectName, Timestamp, FrameNumber, FaceBlendShapes, DeviceId);
-		}
 	}
 }
 
 PRAGMA_ENABLE_DEPRECATION_WARNINGS
 
-bool FAppleARKitLiveLinkSource::Exec(UWorld* InWorld, const TCHAR* Cmd, FOutputDevice& Ar)
-{
-	if (FParse::Command(&Cmd, TEXT("LiveLinkFaceAR")))
-	{
-		FString RemoteIp;
-		if (FParse::Value(Cmd, TEXT("SendTo="), RemoteIp))
-		{
-			// We need to recreate the LiveLink remote publisher
-			RemoteLiveLinkPublisher = nullptr;
-#if PLATFORM_IOS
-			// Only send from iOS to desktop
-			// This will perform the sending of the data to the remote
-			FAppleARKitLiveLinkRemotePublisher* Publisher = new FAppleARKitLiveLinkRemotePublisher(RemoteIp);
-			if (Publisher->InitSendSocket())
-			{
-				RemoteLiveLinkPublisher = MakeShareable(Publisher);
-			}
-			else
-			{
-				UE_LOG(LogAppleARKitFace, Warning, TEXT("Failed to create LiveLink remote publisher, so no data will be sent out"));
-				delete Publisher;
-			}
-#endif
-			return true;
-		}
-	}
-	return false;
-}
-
 // 1 = Initial version
 // 2 = ARKit 2.0 extra blendshapes
 // 3 = Removed the timestamp to derive locally
@@ -257,6 +232,7 @@ bool FAppleARKitLiveLinkRemotePublisher::InitSendSocket()
 		SendSocket = SocketSubsystem->CreateSocket(NAME_DGram, TEXT("FAppleARKitLiveLinkRemotePublisher socket"), true);
 		SendSocket->SetReuseAddr();
 		SendSocket->SetNonBlocking();
+		UE_LOG(LogAppleARKitFace, Log, TEXT("Sending LiveLink face AR data to address (%s)"), *Addr->ToString(true));
 	}
 	return SendSocket != nullptr;
 }
@@ -277,7 +253,6 @@ TSharedRef<FInternetAddr> FAppleARKitLiveLinkRemotePublisher::GetSendAddress()
 		SendAddr->SetPort(LiveLinkPort);
 		bool bIsValid = false;
 		SendAddr->SetIp(*RemoteIp, bIsValid);
-		UE_LOG(LogAppleARKitFace, Log, TEXT("Sending LiveLink face AR data to address (%s)"), *SendAddr->ToString(true));
 	}
 	return SendAddr;
 }
@@ -366,7 +341,7 @@ void FAppleARKitLiveLinkRemoteListener::InitLiveLinkSource()
 {
 	if (!Source.IsValid())
 	{
-		Source = FAppleARKitLiveLinkSourceFactory::CreateLiveLinkSource(false);
+		Source = FAppleARKitLiveLinkSourceFactory::CreateLiveLinkSource();
 	}
 }
 
diff --git a/Engine/Plugins/Runtime/AR/Apple/AppleARKitFaceSupport/Source/AppleARKitFaceSupport/Private/AppleARKitLiveLinkSource.h b/Engine/Plugins/Runtime/AR/Apple/AppleARKitFaceSupport/Source/AppleARKitFaceSupport/Private/AppleARKitLiveLinkSource.h
--- a/Engine/Plugins/Runtime/AR/Apple/AppleARKitFaceSupport/Source/AppleARKitFaceSupport/Private/AppleARKitLiveLinkSource.h
+++ b/Engine/Plugins/Runtime/AR/Apple/AppleARKitFaceSupport/Source/AppleARKitFaceSupport/Private/AppleARKitLiveLinkSource.h
@@ -751,11 +751,10 @@ private:
 
 /** Publishes face blend shapes to LiveLink for use locally */
 class FAppleARKitLiveLinkSource :
-	public ILiveLinkSourceARKit,
-	public FSelfRegisteringExec
+	public ILiveLinkSourceARKit
 {
 public:
-	FAppleARKitLiveLinkSource(bool bCreateRemotePublisher);
+	FAppleARKitLiveLinkSource();
 	virtual ~FAppleARKitLiveLinkSource() {}
 
 private:
@@ -772,19 +771,12 @@ private:
 	virtual void PublishBlendShapes(FName SubjectName, double Timestamp, uint32 FrameNumber, const FARBlendShapeMap& FaceBlendShapes, FName DeviceId = NAME_None) override;
 	// End IARKitBlendShapePublisher
 
-	//~ FSelfRegisteringExec
-	virtual bool Exec(UWorld* InWorld, const TCHAR* Cmd, FOutputDevice& Ar) override;
-	//~ FSelfRegisteringExec
-
 	/** The local client to push data updates to */
 	ILiveLinkClient* Client;
 
 	/** Our identifier in LiveLink */
 	FGuid SourceGuid;
 
-	/** Provider to push data to remote listeners */
-	TSharedPtr<IARKitBlendShapePublisher> RemoteLiveLinkPublisher;
-
 	/** The last time we sent the data. Used to not send redundant data */
 	uint32 LastFramePublished;
 
diff --git a/Engine/Plugins/Runtime/AR/Apple/AppleARKitFaceSupport/Source/AppleARKitFaceSupport/Public/AppleARKitLiveLinkSourceFactory.h b/Engine/Plugins/Runtime/AR/Apple/AppleARKitFaceSupport/Source/AppleARKitFaceSupport/Public/AppleARKitLiveLinkSourceFactory.h
--- a/Engine/Plugins/Runtime/AR/Apple/AppleARKitFaceSupport/Source/AppleARKitFaceSupport/Public/AppleARKitLiveLinkSourceFactory.h
+++ b/Engine/Plugins/Runtime/AR/Apple/AppleARKitFaceSupport/Source/AppleARKitFaceSupport/Public/AppleARKitLiveLinkSourceFactory.h
@@ -26,8 +26,14 @@ class APPLEARKITFACESUPPORT_API FAppleARKitLiveLinkSourceFactory
 {
 public:
 	/** Creates a face mesh source that will autobind to the tracked face mesh */
-	static TSharedPtr<ILiveLinkSourceARKit> CreateLiveLinkSource(bool bCreateRemotePublisher);
+	static TSharedPtr<ILiveLinkSourceARKit> CreateLiveLinkSource();
 
-	/** Creates the listener that will receive remote events and publish them locally */
+	/** Creates the singleton listener that will receive remote events and publish them locally */
 	static void CreateLiveLinkRemoteListener();
+
+	/** Creates the publisher that will send remote events to a specified IP */
+	static TSharedPtr<IARKitBlendShapePublisher, ESPMode::ThreadSafe> CreateLiveLinkRemotePublisher(const FString& RemoteAddr = FString());
+
+	/** Creates the publisher that will write the curve data to disk */
+	static TSharedPtr<IARKitBlendShapePublisher, ESPMode::ThreadSafe> CreateLiveLinkLocalFileWriter();
 };
