Added a file writer for face tracking that logs each frame's data out per file or as a single file in either CSV or JSON formats

diff --git a/Engine/Plugins/Runtime/AR/Apple/AppleARKit/Source/AppleARKit/Public/AppleARKitConversion.h b/Engine/Plugins/Runtime/AR/Apple/AppleARKit/Source/AppleARKit/Public/AppleARKitConversion.h
--- a/Engine/Plugins/Runtime/AR/Apple/AppleARKit/Source/AppleARKit/Public/AppleARKitConversion.h
+++ b/Engine/Plugins/Runtime/AR/Apple/AppleARKit/Source/AppleARKit/Public/AppleARKitConversion.h
@@ -323,6 +323,7 @@ struct FAppleARKitAnchorData
 	{
 		if (this != &Other)
 		{
+			Clear();
 			Copy(Other);
 		}
 
@@ -354,6 +355,14 @@ struct FAppleARKitAnchorData
 		bIsTracked = Other.bIsTracked;
 	}
 
+	void Clear()
+	{
+		BoundaryVerts.Empty();
+		BlendShapes.Empty();
+		FaceVerts.Empty();
+		ProbeTexture = nullptr;
+	}
+
 	FTransform Transform;
 	EAppleAnchorType AnchorType;
 	FGuid AnchorGUID;
diff --git a/Engine/Plugins/Runtime/AR/Apple/AppleARKitFaceSupport/Source/AppleARKitFaceSupport/Private/AppleARKitFaceMeshConversion.h b/Engine/Plugins/Runtime/AR/Apple/AppleARKitFaceSupport/Source/AppleARKitFaceSupport/Private/AppleARKitFaceMeshConversion.h
--- a/Engine/Plugins/Runtime/AR/Apple/AppleARKitFaceSupport/Source/AppleARKitFaceSupport/Private/AppleARKitFaceMeshConversion.h
+++ b/Engine/Plugins/Runtime/AR/Apple/AppleARKitFaceSupport/Source/AppleARKitFaceSupport/Private/AppleARKitFaceMeshConversion.h
@@ -83,7 +83,7 @@ static FORCEINLINE FARBlendShapeMap ToBlendShapeMap(NSDictionary<ARBlendShapeLoc
 #define SET_BLEND_SHAPE(AppleShape, UE4Shape) \
 	if (BlendShapes[AppleShape]) \
 	{ \
-		BlendShapeMap.Add(UE4Shape, [BlendShapes[AppleShape] floatValue]); \
+		BlendShapeMap.Add(UE4Shape, FMath::Max([BlendShapes[AppleShape] floatValue], 0.f)); \
 	} \
 	else \
 	{ \
diff --git a/Engine/Plugins/Runtime/AR/Apple/AppleARKitFaceSupport/Source/AppleARKitFaceSupport/Private/AppleARKitFaceSupportImpl.cpp b/Engine/Plugins/Runtime/AR/Apple/AppleARKitFaceSupport/Source/AppleARKitFaceSupport/Private/AppleARKitFaceSupportImpl.cpp
--- a/Engine/Plugins/Runtime/AR/Apple/AppleARKitFaceSupport/Source/AppleARKitFaceSupport/Private/AppleARKitFaceSupportImpl.cpp
+++ b/Engine/Plugins/Runtime/AR/Apple/AppleARKitFaceSupport/Source/AppleARKitFaceSupport/Private/AppleARKitFaceSupportImpl.cpp
@@ -101,6 +101,21 @@ bool FAppleARKitFaceSupport::Exec(UWorld*, const TCHAR* Cmd, FOutputDevice& Ar)
 	return false;
 }
 
+void FAppleARKitFaceSupport::InitRealtimeProviders()
+{
+	static bool bNeedsInit = true;
+	if (bNeedsInit)
+	{
+		bNeedsInit = false;
+#if PLATFORM_IOS
+		// This will perform the sending of the data to the remote
+		RemoteLiveLinkPublisher = FAppleARKitLiveLinkSourceFactory::CreateLiveLinkRemotePublisher();
+		// Create the file writer if required. Will return nullptr if not configured
+		LiveLinkFileWriter = FAppleARKitLiveLinkSourceFactory::CreateLiveLinkLocalFileWriter();
+#endif
+	}
+}
+
 #if SUPPORTS_ARKIT_1_0
 
 ARConfiguration* FAppleARKitFaceSupport::ToARConfiguration(UARSessionConfig* SessionConfig)
@@ -115,6 +130,9 @@ ARConfiguration* FAppleARKitFaceSupport::ToARConfiguration(UARSessionConfig* Ses
 		SessionConfiguration = [ARFaceTrackingConfiguration new];
 	}
 
+	// Init the remote sender and file loggers if requested
+	InitRealtimeProviders();
+
 	// Copy / convert properties
 	SessionConfiguration.lightEstimationEnabled = SessionConfig->GetLightEstimationMode() != EARLightEstimationMode::None;
 	SessionConfiguration.providesAudioData = NO;
@@ -154,12 +172,12 @@ TArray<TSharedPtr<FAppleARKitAnchorData>> FAppleARKitFaceSupport::MakeAnchorData
 void FAppleARKitFaceSupport::ProcessRealTimePublishers(TSharedPtr<FAppleARKitAnchorData> AnchorData)
 {
 	// Copy the data from the passed in anchor
-	AsyncAnchorCopy = *AnchorData;
-	AsyncTask(ENamedThreads::AnyBackgroundThreadNormalTask, [this]()
+    TSharedPtr<FAppleARKitAnchorData> AsyncAnchorCopy = MakeShared<FAppleARKitAnchorData>(*AnchorData);
+	AsyncTask(ENamedThreads::AnyBackgroundThreadNormalTask, [this, AsyncAnchorCopy]()
 	{
-		double Timestamp = AsyncAnchorCopy.Timestamp;
-		uint32 FrameNumber = AsyncAnchorCopy.FrameNumber;
-		const FARBlendShapeMap& BlendShapes = AsyncAnchorCopy.BlendShapes;
+		double Timestamp = AsyncAnchorCopy->Timestamp;
+		uint32 FrameNumber = AsyncAnchorCopy->FrameNumber;
+		const FARBlendShapeMap& BlendShapes = AsyncAnchorCopy->BlendShapes;
 
 		if (RemoteLiveLinkPublisher.IsValid())
 		{
@@ -183,12 +201,12 @@ void FAppleARKitFaceSupport::PublishLiveLinkData(TSharedPtr<FAppleARKitAnchorDat
 		if (GetDefault<UAppleARKitSettings>()->bEnableLiveLinkForFaceTracking)
 		{
 			FaceTrackingLiveLinkSubjectName = GetDefault<UAppleARKitSettings>()->DefaultFaceTrackingLiveLinkSubjectName;
-	#if PLATFORM_IOS
+#if PLATFORM_IOS
 			LiveLinkSource = FAppleARKitLiveLinkSourceFactory::CreateLiveLinkSource();
-	#else
+#else
 			// This should be started already, but just in case
 			FAppleARKitLiveLinkSourceFactory::CreateLiveLinkRemoteListener();
-	#endif
+#endif
 		}
 	}
 
diff --git a/Engine/Plugins/Runtime/AR/Apple/AppleARKitFaceSupport/Source/AppleARKitFaceSupport/Private/AppleARKitFaceSupportImpl.h b/Engine/Plugins/Runtime/AR/Apple/AppleARKitFaceSupport/Source/AppleARKitFaceSupport/Private/AppleARKitFaceSupportImpl.h
--- a/Engine/Plugins/Runtime/AR/Apple/AppleARKitFaceSupport/Source/AppleARKitFaceSupport/Private/AppleARKitFaceSupportImpl.h
+++ b/Engine/Plugins/Runtime/AR/Apple/AppleARKitFaceSupport/Source/AppleARKitFaceSupport/Private/AppleARKitFaceSupportImpl.h
@@ -35,6 +35,8 @@ private:
 	/** Publishes the remote publisher and the file writer if present */
 	void ProcessRealTimePublishers(TSharedPtr<FAppleARKitAnchorData> AnchorData);
 #endif
+	/** Inits the real time providers if needed */
+	void InitRealtimeProviders();
 
 	//~ FSelfRegisteringExec
 	virtual bool Exec(UWorld*, const TCHAR* Cmd, FOutputDevice& Ar) override;
@@ -50,9 +52,5 @@ private:
 	TSharedPtr<IARKitBlendShapePublisher, ESPMode::ThreadSafe> RemoteLiveLinkPublisher;
 	/** A publisher that writes the data to disk */
 	TSharedPtr<IARKitBlendShapePublisher, ESPMode::ThreadSafe> LiveLinkFileWriter;
-#if SUPPORTS_ARKIT_1_0
-	/** Used to hold a copy of the data to be processed by the background thread */
-	FAppleARKitAnchorData AsyncAnchorCopy;
-#endif
 };
 
diff --git a/Engine/Plugins/Runtime/AR/Apple/AppleARKitFaceSupport/Source/AppleARKitFaceSupport/Private/AppleARKitLiveLinkSource.cpp b/Engine/Plugins/Runtime/AR/Apple/AppleARKitFaceSupport/Source/AppleARKitFaceSupport/Private/AppleARKitLiveLinkSource.cpp
--- a/Engine/Plugins/Runtime/AR/Apple/AppleARKitFaceSupport/Source/AppleARKitFaceSupport/Private/AppleARKitLiveLinkSource.cpp
+++ b/Engine/Plugins/Runtime/AR/Apple/AppleARKitFaceSupport/Source/AppleARKitFaceSupport/Private/AppleARKitLiveLinkSource.cpp
@@ -9,6 +9,7 @@
 #include "ARBlueprintLibrary.h"
 #include "AppleARKitModule.h"
 #include "Features/IModularFeatures.h"
+#include "Misc/FileHelper.h"
 
 #include "SocketSubsystem.h"
 #include "IPAddress.h"
@@ -84,6 +85,23 @@ TSharedPtr<IARKitBlendShapePublisher, ESPMode::ThreadSafe> FAppleARKitLiveLinkSo
 TSharedPtr<IARKitBlendShapePublisher, ESPMode::ThreadSafe> FAppleARKitLiveLinkSourceFactory::CreateLiveLinkLocalFileWriter()
 {
 	TSharedPtr<IARKitBlendShapePublisher, ESPMode::ThreadSafe> LocalFileWriter;
+	bool bWantsFaceTrackingFileWriter = false;
+	GConfig->GetBool(TEXT("/Script/AppleARKit.AppleARKitSettings"), TEXT("bWantsFaceTrackingFileWriter"), bWantsFaceTrackingFileWriter, GEngineIni);
+	if (bWantsFaceTrackingFileWriter)
+	{
+		bool bCsvOrJson = true;
+		GConfig->GetBool(TEXT("/Script/AppleARKit.AppleARKitSettings"), TEXT("bWantsCsvOrJsonFileWriter"), bCsvOrJson, GEngineIni);
+		FAppleARKitLiveLinkFileWriter* FileWriter = nullptr;
+		if (bCsvOrJson)
+		{
+			FileWriter = new FAppleARKitLiveLinkFileWriterCsv();
+		}
+		else
+		{
+			FileWriter = new FAppleARKitLiveLinkFileWriterJson();
+		}
+		LocalFileWriter = MakeShareable(FileWriter);
+	}
 	return LocalFileWriter;
 }
 
@@ -412,3 +430,173 @@ void FAppleARKitLiveLinkRemoteListener::Tick(float DeltaTime)
 		}
 	}
 }
+
+FAppleARKitLiveLinkFileWriter::FAppleARKitLiveLinkFileWriter(const TCHAR* InFileExtension)
+	: FileExtension(InFileExtension)
+	, bSavePerFrameOrOnDemand(false)
+{
+	// Read the config values for this
+	GConfig->GetBool(TEXT("/Script/AppleARKit.AppleARKitSettings"), TEXT("bFaceTrackingWriteEachFrame"), bSavePerFrameOrOnDemand, GEngineIni);
+}
+
+FAppleARKitLiveLinkFileWriter::~FAppleARKitLiveLinkFileWriter()
+{
+	// Save on close if desired
+	if (!bSavePerFrameOrOnDemand)
+	{
+		SaveFileData();
+	}
+}
+
+void FAppleARKitLiveLinkFileWriter::SaveFileData()
+{
+	FString SaveData;
+
+	SaveData = BuildSaveData();
+	// Write the data to the user directory
+	FFileHelper::SaveStringToFile(SaveData, *GenerateFilePath(), FFileHelper::EEncodingOptions::ForceAnsi);
+
+	FrameHistory.Empty();
+}
+
+FString FAppleARKitLiveLinkFileWriter::GenerateFilePath()
+{
+	FDateTime DateTime = FDateTime::UtcNow();
+	const FString UserDir = FPlatformProcess::UserDir();
+	const FString DeviceNameString = DeviceName.ToString();
+	return FString::Printf(TEXT("%sFaceTracking/%s_%d-%d-%d-%d-%d-%d-%d%s"), *UserDir, *DeviceNameString,
+		DateTime.GetYear(), DateTime.GetMonth(), DateTime.GetDay(), DateTime.GetHour(), DateTime.GetMinute(), DateTime.GetSecond(), DateTime.GetMillisecond(),
+		*FileExtension);
+}
+
+void FAppleARKitLiveLinkFileWriter::PublishBlendShapes(FName SubjectName, double Timestamp, uint32 FrameNumber, const FARBlendShapeMap& FaceBlendShapes, FName DeviceId)
+{
+	FScopeLock ScopeLock(&CriticalSection);
+
+	DeviceName = DeviceId;
+	// Add to the array for long running save
+	new(FrameHistory) FFaceTrackingFrame(Timestamp, FrameNumber, FaceBlendShapes);
+
+	if (bSavePerFrameOrOnDemand)
+	{
+		SaveFileData();
+	}
+}
+
+bool FAppleARKitLiveLinkFileWriter::Exec(UWorld*, const TCHAR* Cmd, FOutputDevice& Ar)
+{
+	if (FParse::Command(&Cmd, TEXT("FaceAR")) &&
+		FParse::Command(&Cmd, TEXT("WriteCurveFile")))
+	{
+		FScopeLock ScopeLock(&CriticalSection);
+		SaveFileData();
+		return true;
+	}
+	return false;
+}
+
+FAppleARKitLiveLinkFileWriterCsv::FAppleARKitLiveLinkFileWriterCsv()
+	: FAppleARKitLiveLinkFileWriter(TEXT(".csv"))
+{
+	// Touching UObjects, so needs to be game thread
+	check(IsInGameThread());
+
+	CsvFrameHeader = TEXT("Timestamp, FrameNumber");
+	const UEnum *EnumPtr = FindObject<UEnum>(ANY_PACKAGE, TEXT("EARFaceBlendShape"), true);
+	if (EnumPtr != nullptr)
+	{
+		// Iterate through all of the enum values generating strings for them for CSV/JSON generation
+		for (int32 Shape = 0; Shape < (int32) EARFaceBlendShape::MAX; Shape++)
+		{
+			FName EnumName = ParseEnumName(EnumPtr->GetNameByValue(Shape));
+			FString EnumString(EnumName.ToString());
+			CsvFrameHeader += TEXT(", ") + EnumString;
+		}
+	}
+	CsvFrameHeader += TEXT("\r\n");
+}
+
+FString FAppleARKitLiveLinkFileWriterCsv::BuildCsvRow(const FFaceTrackingFrame& Frame)
+{
+	FString SaveData = FString::Printf(TEXT("%f, %d"), Frame.Timestamp, Frame.FrameNumber);
+	// Add all of the blend shapes on
+	for (int32 Shape = 0; Shape < (int32)EARFaceBlendShape::MAX; Shape++)
+	{
+		float Value = Frame.BlendShapes.FindChecked((EARFaceBlendShape)Shape);
+		SaveData += FString::Printf(TEXT(", %.3f"), Value);
+	}
+	SaveData += TEXT("\r\n");
+	return SaveData;
+}
+
+FString FAppleARKitLiveLinkFileWriterCsv::BuildSaveData()
+{
+	FString SaveData(CsvFrameHeader);
+
+	// Iterate through the array building our string up
+	for (const FFaceTrackingFrame& Frame : FrameHistory)
+	{
+		SaveData += BuildCsvRow(Frame);
+	}
+
+	return SaveData;
+}
+
+FAppleARKitLiveLinkFileWriterJson::FAppleARKitLiveLinkFileWriterJson()
+	: FAppleARKitLiveLinkFileWriter(TEXT(".json"))
+{
+	// Touching UObjects, so needs to be game thread
+	check(IsInGameThread());
+
+	const UEnum *EnumPtr = FindObject<UEnum>(ANY_PACKAGE, TEXT("EARFaceBlendShape"), true);
+	if (EnumPtr != nullptr)
+	{
+		// Iterate through all of the enum values generating strings for them for CSV/JSON generation
+		for (int32 Shape = 0; Shape < (int32) EARFaceBlendShape::MAX; Shape++)
+		{
+			FName EnumName = ParseEnumName(EnumPtr->GetNameByValue(Shape));
+			FString EnumString(EnumName.ToString());
+			BlendShapeJsonKeyNames.Add(EnumString);
+		}
+	}
+}
+
+FString FAppleARKitLiveLinkFileWriterJson::BuildJsonRow(const FFaceTrackingFrame& Frame)
+{
+	FString SaveData = FString::Printf(TEXT("\t{\r\n\t\t\"Timestamp\" : %f,\r\n\t\t\"FrameNumber\" : %d,\r\n"), Frame.Timestamp, Frame.FrameNumber);
+	bool bNeedsComma = false;
+	// Add all of the blend shapes on
+	for (int32 Shape = 0; Shape < (int32)EARFaceBlendShape::MAX; Shape++)
+	{
+		if (bNeedsComma)
+		{
+			SaveData += TEXT(",\r\n");
+		}
+		float Value = Frame.BlendShapes.FindChecked((EARFaceBlendShape)Shape);
+		SaveData += FString::Printf(TEXT("\t\t\"%s\" : %.3f"), *BlendShapeJsonKeyNames[Shape], Value);
+		bNeedsComma = true;
+	}
+	SaveData += TEXT("\r\n\t}");
+	return SaveData;
+}
+
+FString FAppleARKitLiveLinkFileWriterJson::BuildSaveData()
+{
+	FString SaveData;
+	SaveData += TEXT("{\r\n");
+	SaveData += TEXT("\t\"Frames\" : [\r\n");
+	bool bNeedsComma = false;
+	// Add each frame
+	for (const FFaceTrackingFrame& Frame : FrameHistory)
+	{
+		if (bNeedsComma)
+		{
+			SaveData += TEXT(",\r\n");
+		}
+		SaveData += BuildJsonRow(Frame);
+		bNeedsComma = true;
+	}
+	SaveData += TEXT("\t]\r\n");
+	SaveData += TEXT("}\r\n");
+	return SaveData;
+}
diff --git a/Engine/Plugins/Runtime/AR/Apple/AppleARKitFaceSupport/Source/AppleARKitFaceSupport/Private/AppleARKitLiveLinkSource.h b/Engine/Plugins/Runtime/AR/Apple/AppleARKitFaceSupport/Source/AppleARKitFaceSupport/Private/AppleARKitLiveLinkSource.h
--- a/Engine/Plugins/Runtime/AR/Apple/AppleARKitFaceSupport/Source/AppleARKitFaceSupport/Private/AppleARKitLiveLinkSource.h
+++ b/Engine/Plugins/Runtime/AR/Apple/AppleARKitFaceSupport/Source/AppleARKitFaceSupport/Private/AppleARKitLiveLinkSource.h
@@ -749,6 +749,90 @@ private:
 	FNboSerializeToBuffer SendBuffer;
 };
 
+/** Writes files to disk with the face tracking information */
+class FAppleARKitLiveLinkFileWriter :
+	public FSelfRegisteringExec,
+	public IARKitBlendShapePublisher
+{
+public:
+	virtual ~FAppleARKitLiveLinkFileWriter();
+
+protected:
+	FAppleARKitLiveLinkFileWriter(const TCHAR* InFileExtension);
+
+	struct FFaceTrackingFrame
+	{
+		FFaceTrackingFrame(double InTimestamp, uint32 InFrameNumber, const FARBlendShapeMap& InBlendShapes)
+			: Timestamp(InTimestamp)
+			, FrameNumber(InFrameNumber)
+			, BlendShapes(InBlendShapes)
+		{
+		}
+
+		double Timestamp;
+		uint32 FrameNumber;
+		FARBlendShapeMap BlendShapes;
+	};
+
+	// IARKitBlendShapePublisher interface
+	virtual void PublishBlendShapes(FName SubjectName, double Timestamp, uint32 FrameNumber, const FARBlendShapeMap& FaceBlendShapes, FName DeviceId = NAME_None) override;
+	// End IARKitBlendShapePublisher
+
+	//~ FSelfRegisteringExec
+	virtual bool Exec(UWorld*, const TCHAR* Cmd, FOutputDevice& Ar) override;
+	//~ FSelfRegisteringExec
+
+	virtual FString BuildSaveData() = 0;
+
+	void SaveFileData();
+	FString GenerateFilePath();
+
+	/** Needed because Exec commands are on the game thread and updates are on an arbitrary thread */
+	FCriticalSection CriticalSection;
+	/** The file extension to write to */
+	FString FileExtension;
+	/** The list of the last frames seen. Used when rolling up all data in one file */
+	TArray<FFaceTrackingFrame> FrameHistory;
+	/** The name to use when generating the file to save to */
+	FName DeviceName;
+	/** If true, it writes a file each frame. Otherwise, it writes the data upon demand or exit */
+	bool bSavePerFrameOrOnDemand;
+};
+
+class FAppleARKitLiveLinkFileWriterJson :
+	public FAppleARKitLiveLinkFileWriter
+{
+public:
+	FAppleARKitLiveLinkFileWriterJson();
+	virtual ~FAppleARKitLiveLinkFileWriterJson() {}
+
+private:
+	//~ FAppleARKitLiveLinkFileWriter
+	virtual FString BuildSaveData() override;
+	FString BuildJsonRow(const FFaceTrackingFrame& Frame);
+	//~ FAppleARKitLiveLinkFileWriter
+
+	/** A list of the keys to use when generating JSON files, cached because can't access off game thread */
+	TArray<FString> BlendShapeJsonKeyNames;
+};
+
+class FAppleARKitLiveLinkFileWriterCsv :
+	public FAppleARKitLiveLinkFileWriter
+{
+public:
+	FAppleARKitLiveLinkFileWriterCsv();
+	virtual ~FAppleARKitLiveLinkFileWriterCsv() {}
+
+private:
+	//~ FAppleARKitLiveLinkFileWriter
+	virtual FString BuildSaveData() override;
+	FString BuildCsvRow(const FFaceTrackingFrame& Frame);
+	//~ FAppleARKitLiveLinkFileWriter
+
+	/** A pre-made header to use when writing CSV files, cached because can't access off game thread  */
+	FString CsvFrameHeader;
+};
+
 /** Publishes face blend shapes to LiveLink for use locally */
 class FAppleARKitLiveLinkSource :
 	public ILiveLinkSourceARKit
