Cleared the realtime publishers during face ar support shutdown to flush the output sooner

diff --git a/Engine/Plugins/Runtime/AR/Apple/AppleARKitFaceSupport/Source/AppleARKitFaceSupport/Private/AppleARKitFaceSupportImpl.cpp b/Engine/Plugins/Runtime/AR/Apple/AppleARKitFaceSupport/Source/AppleARKitFaceSupport/Private/AppleARKitFaceSupportImpl.cpp
--- a/Engine/Plugins/Runtime/AR/Apple/AppleARKitFaceSupport/Source/AppleARKitFaceSupport/Private/AppleARKitFaceSupportImpl.cpp
+++ b/Engine/Plugins/Runtime/AR/Apple/AppleARKitFaceSupport/Source/AppleARKitFaceSupport/Private/AppleARKitFaceSupportImpl.cpp
@@ -78,6 +78,9 @@ void FAppleARKitFaceSupport::Init()
 
 void FAppleARKitFaceSupport::Shutdown()
 {
+	RemoteLiveLinkPublisher = nullptr;
+	LiveLinkFileWriter = nullptr;
+
 	IModularFeatures::Get().UnregisterModularFeature(GetModularFeatureName(), this);
 }
 
