Removed the face mirroring setting from the arkit settings class since it had a setting in the session config object
Changed the code to use the session config setting

diff --git a/Engine/Plugins/Runtime/AR/Apple/AppleARKitFaceSupport/Source/AppleARKitFaceSupport/Private/AppleARKitFaceMeshConversion.h b/Engine/Plugins/Runtime/AR/Apple/AppleARKitFaceSupport/Source/AppleARKitFaceSupport/Private/AppleARKitFaceMeshConversion.h
--- a/Engine/Plugins/Runtime/AR/Apple/AppleARKitFaceSupport/Source/AppleARKitFaceSupport/Private/AppleARKitFaceMeshConversion.h
+++ b/Engine/Plugins/Runtime/AR/Apple/AppleARKitFaceSupport/Source/AppleARKitFaceSupport/Private/AppleARKitFaceMeshConversion.h
@@ -35,11 +35,8 @@ static FORCEINLINE TArray<FVector> ToVertexBuffer(const vector_float3* Vertices,
 	return VertexBuffer;
 }
 
-static FORCEINLINE FARBlendShapeMap ToBlendShapeMap(NSDictionary<ARBlendShapeLocation,NSNumber *>* BlendShapes, const FTransform& Transform, const FTransform& LeftEyeTransform, const FTransform& RightEyeTransform)
+static FORCEINLINE FARBlendShapeMap ToBlendShapeMap(bool bFaceMirrored, NSDictionary<ARBlendShapeLocation,NSNumber *>* BlendShapes, const FTransform& Transform, const FTransform& LeftEyeTransform, const FTransform& RightEyeTransform)
 {
-	// Do we want to capture face performance or look at the face as if in a mirror (Apple is mirrored so we mirror the mirror)
-	const bool bFaceMirrored = GetDefault<UAppleARKitSettings>()->DefaultFaceTrackingDirection == EARFaceTrackingDirection::FaceMirrored;
-	
 	FARBlendShapeMap BlendShapeMap;
 	FRotator TrackedRot(Transform.GetRotation());
 	// Map the -180..180 range to -1..1
diff --git a/Engine/Plugins/Runtime/AR/Apple/AppleARKitFaceSupport/Source/AppleARKitFaceSupport/Private/AppleARKitFaceSupportImpl.cpp b/Engine/Plugins/Runtime/AR/Apple/AppleARKitFaceSupport/Source/AppleARKitFaceSupport/Private/AppleARKitFaceSupportImpl.cpp
--- a/Engine/Plugins/Runtime/AR/Apple/AppleARKitFaceSupport/Source/AppleARKitFaceSupport/Private/AppleARKitFaceSupportImpl.cpp
+++ b/Engine/Plugins/Runtime/AR/Apple/AppleARKitFaceSupport/Source/AppleARKitFaceSupport/Private/AppleARKitFaceSupportImpl.cpp
@@ -14,7 +14,7 @@ DECLARE_CYCLE_STAT(TEXT("Conversion"), STAT_FaceAR_Conversion, STATGROUP_FaceAR)
 
 #if SUPPORTS_ARKIT_1_0
 
-static TSharedPtr<FAppleARKitAnchorData> MakeAnchorData(ARAnchor* Anchor, const FRotator& AdjustBy, EARFaceTrackingUpdate UpdateSetting)
+static TSharedPtr<FAppleARKitAnchorData> MakeAnchorData(bool bFaceMirrored, ARAnchor* Anchor, const FRotator& AdjustBy, EARFaceTrackingUpdate UpdateSetting)
 {
 	SCOPE_CYCLE_COUNTER(STAT_FaceAR_Conversion);
 	
@@ -38,7 +38,7 @@ static TSharedPtr<FAppleARKitAnchorData> MakeAnchorData(ARAnchor* Anchor, const
 		NewAnchor = MakeShared<FAppleARKitAnchorData>(
 			FAppleARKitConversion::ToFGuid(FaceAnchor.identifier),
 			FAppleARKitConversion::ToFTransform(FaceAnchor.transform, AdjustBy),
-			ToBlendShapeMap(FaceAnchor.blendShapes, FAppleARKitConversion::ToFTransform(FaceAnchor.transform, AdjustBy), LeftEyeTransform, RightEyeTransform),
+			ToBlendShapeMap(bFaceMirrored, FaceAnchor.blendShapes, FAppleARKitConversion::ToFTransform(FaceAnchor.transform, AdjustBy), LeftEyeTransform, RightEyeTransform),
 			UpdateSetting == EARFaceTrackingUpdate::CurvesAndGeo ? ToVertexBuffer(FaceAnchor.geometry.vertices, FaceAnchor.geometry.vertexCount) : TArray<FVector>(),
 			LeftEyeTransform,
 			RightEyeTransform,
@@ -61,6 +61,7 @@ static TSharedPtr<FAppleARKitAnchorData> MakeAnchorData(ARAnchor* Anchor, const
 
 FAppleARKitFaceSupport::FAppleARKitFaceSupport()
 {
+	bFaceMirrored = false;
 	// Generate our device id
 	LocalDeviceId = FName(*FPlatformMisc::GetDeviceId());
 }
@@ -151,6 +152,9 @@ ARConfiguration* FAppleARKitFaceSupport::ToARConfiguration(UARSessionConfig* Ses
 		}
 	}
 #endif
+	// Do we want to capture face performance or look at the face as if in a mirror (Apple is mirrored so we mirror the mirror)
+	bFaceMirrored = SessionConfig->GetFaceTrackingDirection() == EARFaceTrackingDirection::FaceMirrored;
+
 	return SessionConfiguration;
 }
 
@@ -160,7 +164,7 @@ TArray<TSharedPtr<FAppleARKitAnchorData>> FAppleARKitFaceSupport::MakeAnchorData
 
 	for (ARAnchor* Anchor in Anchors)
 	{
-		TSharedPtr<FAppleARKitAnchorData> AnchorData = ::MakeAnchorData(Anchor, AdjustBy, UpdateSetting);
+		TSharedPtr<FAppleARKitAnchorData> AnchorData = ::MakeAnchorData(bFaceMirrored, Anchor, AdjustBy, UpdateSetting);
 		if (AnchorData.IsValid())
 		{
 			AnchorList.Add(AnchorData);
diff --git a/Engine/Plugins/Runtime/AR/Apple/AppleARKitFaceSupport/Source/AppleARKitFaceSupport/Private/AppleARKitFaceSupportImpl.h b/Engine/Plugins/Runtime/AR/Apple/AppleARKitFaceSupport/Source/AppleARKitFaceSupport/Private/AppleARKitFaceSupportImpl.h
--- a/Engine/Plugins/Runtime/AR/Apple/AppleARKitFaceSupport/Source/AppleARKitFaceSupport/Private/AppleARKitFaceSupportImpl.h
+++ b/Engine/Plugins/Runtime/AR/Apple/AppleARKitFaceSupport/Source/AppleARKitFaceSupport/Private/AppleARKitFaceSupportImpl.h
@@ -42,6 +42,8 @@ private:
 	virtual bool Exec(UWorld*, const TCHAR* Cmd, FOutputDevice& Ar) override;
 	//~ FSelfRegisteringExec
 
+	/** Whether the face data is mirrored or not */
+	bool bFaceMirrored;
 	/** If requested, publishes face ar updates to LiveLink for the animation system to use */
 	TSharedPtr<ILiveLinkSourceARKit> LiveLinkSource;
 	/** Copied from the UARSessionConfig project settings object */
