Added Timecode support to face tracking frame data
Added a default timecode provider for ARKit for when it's not coming from an external source
Changed image and json/csv logging of face tracking data to generate the file name based off of the timecode
Modified the network protocol to stream the timecode and frame rate of the timecode source to the editor
Added some helper methods for FTimecode and FQualifiedFrameTime

diff --git a/Engine/Plugins/Runtime/AR/Apple/AppleARKit/Source/AppleARKit/Private/AppleARKitSystem.cpp b/Engine/Plugins/Runtime/AR/Apple/AppleARKit/Source/AppleARKit/Private/AppleARKitSystem.cpp
--- a/Engine/Plugins/Runtime/AR/Apple/AppleARKit/Source/AppleARKit/Private/AppleARKitSystem.cpp
+++ b/Engine/Plugins/Runtime/AR/Apple/AppleARKit/Source/AppleARKit/Private/AppleARKitSystem.cpp
@@ -1105,6 +1105,7 @@ void FAppleARKitSystem::AddReferencedObjects( FReferenceCollector& Collector )
 	Collector.AddReferencedObject( CameraDepth );
 	Collector.AddReferencedObjects( CandidateImages );
 	Collector.AddReferencedObjects( CandidateObjects );
+	Collector.AddReferencedObject( TimecodeProvider );
 
 	if(LightEstimate)
 	{
@@ -1192,6 +1193,8 @@ void FAppleARKitSystem::SetDeviceOrientation( EScreenOrientation::Type InOrienta
 PRAGMA_DISABLE_OPTIMIZATION
 bool FAppleARKitSystem::Run(UARSessionConfig* SessionConfig)
 {
+	TimecodeProvider = UAppleARKitSettings::GetTimecodeProvider();
+
 	{
 		// Clear out any existing frames since they aren't valid anymore
 		FScopeLock ScopeLock(&FrameLock);
@@ -1224,7 +1227,7 @@ bool FAppleARKitSystem::Run(UARSessionConfig* SessionConfig)
 		}
 		else
 		{
-			Configuration = FaceARSupport->ToARConfiguration(SessionConfig);
+			Configuration = FaceARSupport->ToARConfiguration(SessionConfig, TimecodeProvider);
 		}
 
 		// Not all session types are supported by all devices
@@ -1495,7 +1498,7 @@ void FAppleARKitSystem::SessionDidAddAnchors_DelegateThread( NSArray<ARAnchor*>*
 
 	// Make sure all anchors get the same timestamp and frame number
 	double Timestamp = FPlatformTime::Seconds();
-	uint32 FrameNumber = FARKitFrameCounter::Get().GetFrameNumber();
+	uint32 FrameNumber = TimecodeProvider->GetTimecode().Frames;
 
 	for (ARAnchor* anchor in anchors)
 	{
@@ -1528,7 +1531,7 @@ void FAppleARKitSystem::SessionDidUpdateAnchors_DelegateThread( NSArray<ARAnchor
 
 	// Make sure all anchors get the same timestamp and frame number
 	double Timestamp = FPlatformTime::Seconds();
-	uint32 FrameNumber = FARKitFrameCounter::Get().GetFrameNumber();
+	uint32 FrameNumber = TimecodeProvider->GetTimecode().Frames;
 
 	for (ARAnchor* anchor in anchors)
 	{
@@ -1887,7 +1890,8 @@ void FAppleARKitSystem::WriteCameraImageToDisk(CVPixelBufferRef PixelBuffer)
 	float ImageScale = WrittenCameraImageScale;
 	ETextureRotationDirection ImageRotation = WrittenCameraImageRotation;
 	CIImage* SourceImage = [[CIImage alloc] initWithCVPixelBuffer: PixelBuffer];
-	AsyncTask(ENamedThreads::AnyBackgroundThreadNormalTask, [SourceImage, ImageQuality, ImageScale, ImageRotation]()
+	FTimecode Timecode = TimecodeProvider->GetTimecode();
+	AsyncTask(ENamedThreads::AnyBackgroundThreadNormalTask, [SourceImage, ImageQuality, ImageScale, ImageRotation, Timecode]()
 	{
 		TArray<uint8> JpegBytes;
 		IAppleImageUtilsPlugin::Get().ConvertToJPEG(SourceImage, JpegBytes, ImageQuality, true, true, ImageScale, ImageRotation);
@@ -1896,7 +1900,7 @@ void FAppleARKitSystem::WriteCameraImageToDisk(CVPixelBufferRef PixelBuffer)
 		FDateTime DateTime = FDateTime::UtcNow();
 		static FString UserDir = FPlatformProcess::UserDir();
 		FString FileName = FString::Printf(TEXT("%sCameraImages/Image_%d-%d-%d-%d-%d-%d-%d.jpeg"), *UserDir,
-			DateTime.GetYear(), DateTime.GetMonth(), DateTime.GetDay(), DateTime.GetHour(), DateTime.GetMinute(), DateTime.GetSecond(), DateTime.GetMillisecond());
+			DateTime.GetYear(), DateTime.GetMonth(), DateTime.GetDay(), Timecode.Hours, Timecode.Minutes, Timecode.Seconds, Timecode.Frames);
 		// Write the jpeg to disk
 		FFileHelper::SaveArrayToFile(JpegBytes, *FileName);
 	});
diff --git a/Engine/Plugins/Runtime/AR/Apple/AppleARKit/Source/AppleARKit/Private/AppleARKitSystem.h b/Engine/Plugins/Runtime/AR/Apple/AppleARKit/Source/AppleARKit/Private/AppleARKitSystem.h
--- a/Engine/Plugins/Runtime/AR/Apple/AppleARKit/Source/AppleARKit/Private/AppleARKitSystem.h
+++ b/Engine/Plugins/Runtime/AR/Apple/AppleARKit/Source/AppleARKit/Private/AppleARKitSystem.h
@@ -197,6 +197,9 @@ private:
 
 	// The object that is handling face support if present
 	IAppleARKitFaceSupport* FaceARSupport;
+
+	/** The time code provider to use when tagging time stamps */
+	UTimecodeProvider* TimecodeProvider;
 };
 
 
diff --git a/Engine/Plugins/Runtime/AR/Apple/AppleARKit/Source/AppleARKit/Private/AppleARKitTimecodeProvider.cpp b/Engine/Plugins/Runtime/AR/Apple/AppleARKit/Source/AppleARKit/Private/AppleARKitTimecodeProvider.cpp
new file mode 100644
--- /dev/null
+++ b/Engine/Plugins/Runtime/AR/Apple/AppleARKit/Source/AppleARKit/Private/AppleARKitTimecodeProvider.cpp
@@ -0,0 +1,35 @@
+// Copyright 1998-2018 Epic Games, Inc. All Rights Reserved.
+
+#include "AppleARKitTimecodeProvider.h"
+#include "HAL/PlatformTime.h"
+
+
+UAppleARKitTimecodeProvider::UAppleARKitTimecodeProvider()
+	: FrameRate(60, 1)
+{
+}
+
+bool UAppleARKitTimecodeProvider::Initialize(class UEngine*)
+{
+	return true;
+}
+
+void UAppleARKitTimecodeProvider::Shutdown(class UEngine*)
+{
+}
+
+FTimecode UAppleARKitTimecodeProvider::GetTimecode() const
+{
+	return FTimecode(FPlatformTime::Seconds(), FrameRate);
+}
+
+FFrameRate UAppleARKitTimecodeProvider::GetFrameRate() const
+{
+	return FrameRate;
+}
+
+ETimecodeProviderSynchronizationState UAppleARKitTimecodeProvider::GetSynchronizationState() const
+{
+	return ETimecodeProviderSynchronizationState::Synchronized;
+}
+
diff --git a/Engine/Plugins/Runtime/AR/Apple/AppleARKit/Source/AppleARKit/Private/AppleARKitTimecodeProvider.h b/Engine/Plugins/Runtime/AR/Apple/AppleARKit/Source/AppleARKit/Private/AppleARKitTimecodeProvider.h
new file mode 100644
--- /dev/null
+++ b/Engine/Plugins/Runtime/AR/Apple/AppleARKit/Source/AppleARKit/Private/AppleARKitTimecodeProvider.h
@@ -0,0 +1,33 @@
+// Copyright 1998-2018 Epic Games, Inc. All Rights Reserved.
+
+#pragma once
+
+#include "Engine/TimecodeProvider.h"
+#include "AppleARKitTimecodeProvider.generated.h"
+
+/**
+ * This class is an implementation of the ITimecodeProvider and is used to abstract
+ * out the calculation of the frame & time for an update
+ */
+UCLASS()
+class UAppleARKitTimecodeProvider :
+	public UTimecodeProvider
+{
+	GENERATED_BODY()
+
+public:
+	UAppleARKitTimecodeProvider();
+
+	// ~UTimecodeProvider
+	virtual FTimecode GetTimecode() const override;
+	virtual FFrameRate GetFrameRate() const override;
+	virtual ETimecodeProviderSynchronizationState GetSynchronizationState() const override;
+	virtual bool Initialize(class UEngine* InEngine) override;
+	virtual void Shutdown(class UEngine* InEngine) override;
+	// ~UTimecodeProvider
+
+private:
+	/** The frame rate of updates is 60 hz */
+	FFrameRate FrameRate;
+};
+
diff --git a/Engine/Plugins/Runtime/AR/Apple/AppleARKit/Source/AppleARKit/Public/AppleARKitConversion.h b/Engine/Plugins/Runtime/AR/Apple/AppleARKit/Source/AppleARKit/Public/AppleARKitConversion.h
--- a/Engine/Plugins/Runtime/AR/Apple/AppleARKit/Source/AppleARKit/Public/AppleARKitConversion.h
+++ b/Engine/Plugins/Runtime/AR/Apple/AppleARKit/Source/AppleARKit/Public/AppleARKitConversion.h
@@ -9,6 +9,7 @@
 #include "ARTypes.h"
 #include "ARTrackable.h"
 #include "ARSessionConfig.h"
+#include "Misc/Timecode.h"
 
 #include "IAppleImageUtilsPlugin.h"
 
@@ -16,46 +17,6 @@
 	#import <ARKit/ARKit.h>
 #endif
 
-/**
- * This class handles incrementing the frame counter. This is needed in case there
- * are both new anchors added and existing ones updated in the same "frame". When
- * this happens, the frame number will remain the same due to the time between
- * adds and updates being shorter than the arkit update rate.
- */
-class FARKitFrameCounter
-{
-public:
-	static FARKitFrameCounter& Get()
-	{
-		static FARKitFrameCounter Singleton;
-		return Singleton;
-	}
-
-	uint32 GetFrameNumber()
-	{
-		uint32 FrameNum = FrameNumber;
-		double TimeSecs = FPlatformTime::Seconds();
-		// ARKit runs at 60hz
-		if (TimeSecs - LastFrameSecs > 0.016)
-		{
-			FrameNumber++;
-			LastFrameSecs = TimeSecs;
-		}
-		return FrameNum;
-	}
-
-private:
-	FARKitFrameCounter()
-		: FrameNumber(0)
-		, LastFrameSecs(0.0)
-	{
-	}
-
-	uint32 FrameNumber;
-	double LastFrameSecs;
-};
-
-
 #define AR_SAVE_WORLD_KEY 0x505A474A
 #define AR_SAVE_WORLD_VER 1
 
@@ -288,7 +249,7 @@ struct FAppleARKitAnchorData
 	{
 	}
 
-	FAppleARKitAnchorData(FGuid InAnchorGuid, FTransform InTransform, FARBlendShapeMap InBlendShapes, TArray<FVector> InFaceVerts, FTransform InLeftEyeTransform, FTransform InRightEyeTransform, FVector InLookAtTarget, double InTimestamp, uint32 InFrameNumber)
+	FAppleARKitAnchorData(FGuid InAnchorGuid, FTransform InTransform, FARBlendShapeMap InBlendShapes, TArray<FVector> InFaceVerts, FTransform InLeftEyeTransform, FTransform InRightEyeTransform, FVector InLookAtTarget, const FTimecode& InTimecode, uint32 InFrameRate)
 		: Transform( InTransform )
 		, AnchorType( EAppleAnchorType::FaceAnchor )
 		, AnchorGUID( InAnchorGuid )
@@ -297,8 +258,8 @@ struct FAppleARKitAnchorData
 		, LeftEyeTransform( InLeftEyeTransform )
 		, RightEyeTransform( InRightEyeTransform )
 		, LookAtTarget( InLookAtTarget )
-		, Timestamp( InTimestamp )
-		, FrameNumber( InFrameNumber )
+		, Timecode( InTimecode )
+		, FrameRate( InFrameRate )
 	{
 	}
 
@@ -351,6 +312,8 @@ struct FAppleARKitAnchorData
 		LookAtTarget = Other.LookAtTarget;
 		Timestamp = Other.Timestamp;
 		FrameNumber = Other.FrameNumber;
+		Timecode = Other.Timecode;
+		FrameRate = Other.FrameRate;
 
 		bIsTracked = Other.bIsTracked;
 	}
@@ -384,6 +347,8 @@ struct FAppleARKitAnchorData
 	FVector LookAtTarget;
 	double Timestamp;
 	uint32 FrameNumber;
+	FTimecode Timecode;
+	uint32 FrameRate;
 
 	/** Only valid for tracked real world objects (face, images) */
 	bool bIsTracked;
diff --git a/Engine/Plugins/Runtime/AR/Apple/AppleARKit/Source/AppleARKit/Public/AppleARKitFaceSupport.h b/Engine/Plugins/Runtime/AR/Apple/AppleARKit/Source/AppleARKit/Public/AppleARKitFaceSupport.h
--- a/Engine/Plugins/Runtime/AR/Apple/AppleARKit/Source/AppleARKit/Public/AppleARKitFaceSupport.h
+++ b/Engine/Plugins/Runtime/AR/Apple/AppleARKit/Source/AppleARKit/Public/AppleARKitFaceSupport.h
@@ -15,6 +15,7 @@
 
 struct FAppleARKitAnchorData;
 struct FARVideoFormat;
+class UTimecodeProvider;
 
 class APPLEARKIT_API IAppleARKitFaceSupport :
 	public IModularFeature
@@ -46,8 +47,9 @@ public:
 	 * Creates a face ar specific configuration object if that is requested without exposing the main code to the face APIs
 	 *
 	 * @param SessionConfig the UE4 configuration object that needs processing
+	 * @param InProvider the custom timecode provider to use
 	 */
-	virtual ARConfiguration* ToARConfiguration(UARSessionConfig* SessionConfig) { return nullptr; }
+	virtual ARConfiguration* ToARConfiguration(UARSessionConfig* SessionConfig, UTimecodeProvider* InProvider) { return nullptr; }
 
 	/**
 	 * @return whether this device supports face ar
diff --git a/Engine/Plugins/Runtime/AR/Apple/AppleARKit/Source/AppleARKit/Public/AppleARKitSettings.h b/Engine/Plugins/Runtime/AR/Apple/AppleARKit/Source/AppleARKit/Public/AppleARKitSettings.h
--- a/Engine/Plugins/Runtime/AR/Apple/AppleARKit/Source/AppleARKit/Public/AppleARKitSettings.h
+++ b/Engine/Plugins/Runtime/AR/Apple/AppleARKit/Source/AppleARKit/Public/AppleARKitSettings.h
@@ -5,8 +5,10 @@
 #include "UObject/ObjectMacros.h"
 #include "UObject/Object.h"
 #include "ARTrackable.h"
+#include "Engine/TimecodeProvider.h"
 
 #include "AppleImageUtilsTypes.h"
+#include "AppleARKitTimecodeProvider.h"
 
 #include "AppleARKitSettings.generated.h"
 
@@ -42,6 +44,7 @@ public:
 		, bAdjustThreadPrioritiesDuringARSession(false)
 		, GameThreadPriorityOverride(47)
 		, RenderThreadPriorityOverride(45)
+		, ARKitTimecodeProvider(TEXT("/Script/AppleARKit.AppleARKitTimecodeProvider"))
 	{
 	}
 
@@ -96,4 +99,32 @@ public:
 	/** The render thread priority to change to when an AR session is running, default is 45 */
 	UPROPERTY(Config, EditAnywhere, Category="AR Settings")
 	int32 RenderThreadPriorityOverride;
+
+	/**
+	 * Used to specify the timecode provider to use when identifying when an update occurred.
+	 * Useful when using external timecode generators to sync multiple devices/machines
+	 */
+	UPROPERTY(Config, EditAnywhere, Category="AR Settings")
+	FString ARKitTimecodeProvider;
+
+	static UTimecodeProvider* GetTimecodeProvider()
+	{
+		const FString& ProviderName = GetDefault<UAppleARKitSettings>()->ARKitTimecodeProvider;
+		UTimecodeProvider* TimecodeProvider = FindObject<UTimecodeProvider>(GEngine, *ProviderName);
+		if (TimecodeProvider == nullptr)
+		{
+			// Try to load the class that was requested
+			UClass* Class = LoadClass<UTimecodeProvider>(nullptr, *ProviderName);
+			if (Class != nullptr)
+			{
+				TimecodeProvider = NewObject<UTimecodeProvider>(GEngine, Class);
+			}
+		}
+		// Create the default one if this failed for some reason
+		if (TimecodeProvider == nullptr)
+		{
+			TimecodeProvider = NewObject<UTimecodeProvider>(GEngine, UAppleARKitTimecodeProvider::StaticClass());
+		}
+		return TimecodeProvider;
+	}
 };
diff --git a/Engine/Plugins/Runtime/AR/Apple/AppleARKitFaceSupport/Source/AppleARKitFaceSupport/Private/AppleARKitFaceMeshComponent.cpp b/Engine/Plugins/Runtime/AR/Apple/AppleARKitFaceSupport/Source/AppleARKitFaceSupport/Private/AppleARKitFaceMeshComponent.cpp
--- a/Engine/Plugins/Runtime/AR/Apple/AppleARKitFaceSupport/Source/AppleARKitFaceSupport/Private/AppleARKitFaceMeshComponent.cpp
+++ b/Engine/Plugins/Runtime/AR/Apple/AppleARKitFaceSupport/Source/AppleARKitFaceSupport/Private/AppleARKitFaceMeshComponent.cpp
@@ -387,12 +387,13 @@ void UAppleARKitFaceMeshComponent::OnRep_RemoteCurves()
 			BlendShapes.Add(Curve.GetBlendShape(), Curve.GetAmountAsFloat());
 		}
 		// We have to manage this since the ar system isn't updating this locally for us
-		LastUpdateFrameNumber++;
 		LastUpdateTimestamp = FPlatformTime::Seconds();
+		FTimecode Timecode = FApp::GetTimecode();
+		FFrameRate FrameRate = FApp::GetTimecodeFrameRate();
 		// If we are publishing to LiveLink, then push the update out
 		if (LiveLinkSource.IsValid())
 		{
-			LiveLinkSource->PublishBlendShapes(LiveLinkSubjectName, LastUpdateTimestamp, LastUpdateFrameNumber, BlendShapes);
+			LiveLinkSource->PublishBlendShapes(LiveLinkSubjectName, Timecode, FrameRate.Numerator, BlendShapes);
 		}
 	}
 }
diff --git a/Engine/Plugins/Runtime/AR/Apple/AppleARKitFaceSupport/Source/AppleARKitFaceSupport/Private/AppleARKitFaceSupportImpl.cpp b/Engine/Plugins/Runtime/AR/Apple/AppleARKitFaceSupport/Source/AppleARKitFaceSupport/Private/AppleARKitFaceSupportImpl.cpp
--- a/Engine/Plugins/Runtime/AR/Apple/AppleARKitFaceSupport/Source/AppleARKitFaceSupport/Private/AppleARKitFaceSupportImpl.cpp
+++ b/Engine/Plugins/Runtime/AR/Apple/AppleARKitFaceSupport/Source/AppleARKitFaceSupport/Private/AppleARKitFaceSupportImpl.cpp
@@ -9,12 +9,13 @@
 #include "Misc/ConfigCacheIni.h"
 #include "AppleARKitFaceSupportModule.h"
 #include "Async/Async.h"
+#include "Engine/TimecodeProvider.h"
 
 DECLARE_CYCLE_STAT(TEXT("Conversion"), STAT_FaceAR_Conversion, STATGROUP_FaceAR);
 
 #if SUPPORTS_ARKIT_1_0
 
-static TSharedPtr<FAppleARKitAnchorData> MakeAnchorData(bool bFaceMirrored, ARAnchor* Anchor, const FRotator& AdjustBy, EARFaceTrackingUpdate UpdateSetting)
+static TSharedPtr<FAppleARKitAnchorData> MakeAnchorData(bool bFaceMirrored, ARAnchor* Anchor, const FRotator& AdjustBy, EARFaceTrackingUpdate UpdateSetting, const FTimecode& Timecode, uint32 FrameRate)
 {
 	SCOPE_CYCLE_COUNTER(STAT_FaceAR_Conversion);
 	
@@ -33,8 +34,6 @@ static TSharedPtr<FAppleARKitAnchorData> MakeAnchorData(bool bFaceMirrored, ARAn
 			LookAtTarget = FAppleARKitConversion::ToFVector(FaceAnchor.lookAtPoint);
 		}
 #endif
-		uint32 FrameNumber = FARKitFrameCounter::Get().GetFrameNumber();
-		double Timestamp = FPlatformTime::Seconds();
 		NewAnchor = MakeShared<FAppleARKitAnchorData>(
 			FAppleARKitConversion::ToFGuid(FaceAnchor.identifier),
 			FAppleARKitConversion::ToFTransform(FaceAnchor.transform, AdjustBy),
@@ -43,8 +42,8 @@ static TSharedPtr<FAppleARKitAnchorData> MakeAnchorData(bool bFaceMirrored, ARAn
 			LeftEyeTransform,
 			RightEyeTransform,
 			LookAtTarget,
-			Timestamp,
-			FrameNumber++
+			Timecode,
+			FrameRate
 		);
         // Only convert from 16bit to 32bit once
         if (UpdateSetting == EARFaceTrackingUpdate::CurvesAndGeo && FAppleARKitAnchorData::FaceIndices.Num() == 0)
@@ -59,7 +58,8 @@ static TSharedPtr<FAppleARKitAnchorData> MakeAnchorData(bool bFaceMirrored, ARAn
 
 #endif
 
-FAppleARKitFaceSupport::FAppleARKitFaceSupport()
+FAppleARKitFaceSupport::FAppleARKitFaceSupport() :
+	TimecodeProvider(nullptr)
 {
 	bFaceMirrored = false;
 	// Generate our device id
@@ -122,8 +122,10 @@ void FAppleARKitFaceSupport::InitRealtimeProviders()
 
 #if SUPPORTS_ARKIT_1_0
 
-ARConfiguration* FAppleARKitFaceSupport::ToARConfiguration(UARSessionConfig* SessionConfig)
+ARConfiguration* FAppleARKitFaceSupport::ToARConfiguration(UARSessionConfig* SessionConfig, UTimecodeProvider* InProvider)
 {
+	TimecodeProvider = InProvider;
+
 	ARFaceTrackingConfiguration* SessionConfiguration = nullptr;
 	if (SessionConfig->GetSessionType() == EARSessionType::Face)
 	{
@@ -136,6 +138,10 @@ ARConfiguration* FAppleARKitFaceSupport::ToARConfiguration(UARSessionConfig* Ses
 
 	// Init the remote sender and file loggers if requested
 	InitRealtimeProviders();
+	if (LiveLinkFileWriter.IsValid())
+	{
+		LiveLinkFileWriter->SetTimecodeProvider(TimecodeProvider);
+	}
 
 	// Copy / convert properties
 	SessionConfiguration.lightEstimationEnabled = SessionConfig->GetLightEstimationMode() != EARLightEstimationMode::None;
@@ -162,9 +168,12 @@ TArray<TSharedPtr<FAppleARKitAnchorData>> FAppleARKitFaceSupport::MakeAnchorData
 {
 	TArray<TSharedPtr<FAppleARKitAnchorData>> AnchorList;
 
+	const FTimecode& Timecode = TimecodeProvider->GetTimecode();
+	const FFrameRate& FrameRate = TimecodeProvider->GetFrameRate();
+
 	for (ARAnchor* Anchor in Anchors)
 	{
-		TSharedPtr<FAppleARKitAnchorData> AnchorData = ::MakeAnchorData(bFaceMirrored, Anchor, AdjustBy, UpdateSetting);
+		TSharedPtr<FAppleARKitAnchorData> AnchorData = ::MakeAnchorData(bFaceMirrored, Anchor, AdjustBy, UpdateSetting, Timecode, FrameRate.Numerator);
 		if (AnchorData.IsValid())
 		{
 			AnchorList.Add(AnchorData);
@@ -182,18 +191,18 @@ void FAppleARKitFaceSupport::ProcessRealTimePublishers(TSharedPtr<FAppleARKitAnc
     TSharedPtr<FAppleARKitAnchorData> AsyncAnchorCopy = MakeShared<FAppleARKitAnchorData>(*AnchorData);
 	AsyncTask(ENamedThreads::AnyBackgroundThreadNormalTask, [this, AsyncAnchorCopy]()
 	{
-		double Timestamp = AsyncAnchorCopy->Timestamp;
-		uint32 FrameNumber = AsyncAnchorCopy->FrameNumber;
+		FTimecode Timecode = AsyncAnchorCopy->Timecode;
+		uint32 FrameRate = AsyncAnchorCopy->FrameRate;
 		const FARBlendShapeMap& BlendShapes = AsyncAnchorCopy->BlendShapes;
 
 		if (RemoteLiveLinkPublisher.IsValid())
 		{
-			RemoteLiveLinkPublisher->PublishBlendShapes(FaceTrackingLiveLinkSubjectName, Timestamp, FrameNumber, BlendShapes, LocalDeviceId);
+			RemoteLiveLinkPublisher->PublishBlendShapes(FaceTrackingLiveLinkSubjectName, Timecode, FrameRate, BlendShapes, LocalDeviceId);
 		}
 
 		if (LiveLinkFileWriter.IsValid())
 		{
-			LiveLinkFileWriter->PublishBlendShapes(FaceTrackingLiveLinkSubjectName, Timestamp, FrameNumber, BlendShapes, LocalDeviceId);
+			LiveLinkFileWriter->PublishBlendShapes(FaceTrackingLiveLinkSubjectName, Timecode, FrameRate, BlendShapes, LocalDeviceId);
 		}
 	});
 }
@@ -219,7 +228,7 @@ void FAppleARKitFaceSupport::PublishLiveLinkData(TSharedPtr<FAppleARKitAnchorDat
 
 	if (LiveLinkSource.IsValid())
 	{
-        LiveLinkSource->PublishBlendShapes(FaceTrackingLiveLinkSubjectName, Anchor->Timestamp, Anchor->FrameNumber, Anchor->BlendShapes, LocalDeviceId);
+        LiveLinkSource->PublishBlendShapes(FaceTrackingLiveLinkSubjectName, Anchor->Timecode, Anchor->FrameRate, Anchor->BlendShapes, LocalDeviceId);
 	}
 }
 
@@ -234,4 +243,3 @@ TArray<FARVideoFormat> FAppleARKitFaceSupport::ToARConfiguration()
 	return FAppleARKitConversion::FromARVideoFormatArray(ARFaceTrackingConfiguration.supportedVideoFormats);
 }
 #endif
-
diff --git a/Engine/Plugins/Runtime/AR/Apple/AppleARKitFaceSupport/Source/AppleARKitFaceSupport/Private/AppleARKitFaceSupportImpl.h b/Engine/Plugins/Runtime/AR/Apple/AppleARKitFaceSupport/Source/AppleARKitFaceSupport/Private/AppleARKitFaceSupportImpl.h
--- a/Engine/Plugins/Runtime/AR/Apple/AppleARKitFaceSupport/Source/AppleARKitFaceSupport/Private/AppleARKitFaceSupportImpl.h
+++ b/Engine/Plugins/Runtime/AR/Apple/AppleARKitFaceSupport/Source/AppleARKitFaceSupport/Private/AppleARKitFaceSupportImpl.h
@@ -24,7 +24,7 @@ private:
 #if SUPPORTS_ARKIT_1_0
 	// ~IAppleARKitFaceSupport
 	virtual TArray<TSharedPtr<FAppleARKitAnchorData>> MakeAnchorData(NSArray<ARAnchor*>* NewAnchors, const FRotator& AdjustBy, EARFaceTrackingUpdate UpdateSetting) override;
-	virtual ARConfiguration* ToARConfiguration(UARSessionConfig* InSessionConfig) override;
+	virtual ARConfiguration* ToARConfiguration(UARSessionConfig* InSessionConfig, UTimecodeProvider* InProvider) override;
 	virtual void PublishLiveLinkData(TSharedPtr<FAppleARKitAnchorData> Anchor) override;
 	virtual bool DoesSupportFaceAR() override;
 #if SUPPORTS_ARKIT_1_5
@@ -54,5 +54,11 @@ private:
 	TSharedPtr<IARKitBlendShapePublisher, ESPMode::ThreadSafe> RemoteLiveLinkPublisher;
 	/** A publisher that writes the data to disk */
 	TSharedPtr<IARKitBlendShapePublisher, ESPMode::ThreadSafe> LiveLinkFileWriter;
+	/**
+	 * The time code provider to use when tagging time stamps
+	 * Note: this requires the FAppleARKitSystem object to mark it in use so GC doesn't destroy it. Normally it would
+	 * implement the FGCObject interface but this gets created before UObjects are init-ed so not possible
+	 * */
+	UTimecodeProvider* TimecodeProvider;
 };
 
diff --git a/Engine/Plugins/Runtime/AR/Apple/AppleARKitFaceSupport/Source/AppleARKitFaceSupport/Private/AppleARKitLiveLinkSource.cpp b/Engine/Plugins/Runtime/AR/Apple/AppleARKitFaceSupport/Source/AppleARKitFaceSupport/Private/AppleARKitLiveLinkSource.cpp
--- a/Engine/Plugins/Runtime/AR/Apple/AppleARKitFaceSupport/Source/AppleARKitFaceSupport/Private/AppleARKitLiveLinkSource.cpp
+++ b/Engine/Plugins/Runtime/AR/Apple/AppleARKitFaceSupport/Source/AppleARKitFaceSupport/Private/AppleARKitLiveLinkSource.cpp
@@ -151,7 +151,7 @@ static FName ParseEnumName(FName EnumName)
 // Temporary for 4.20
 PRAGMA_DISABLE_DEPRECATION_WARNINGS
 
-void FAppleARKitLiveLinkSource::PublishBlendShapes(FName SubjectName, double Timestamp, uint32 FrameNumber, const FARBlendShapeMap& FaceBlendShapes, FName DeviceId)
+void FAppleARKitLiveLinkSource::PublishBlendShapes(FName SubjectName, const FTimecode& Timecode, uint32 FrameRate, const FARBlendShapeMap& FaceBlendShapes, FName DeviceId)
 {
 	SCOPE_CYCLE_COUNTER(STAT_FaceAR_Local_PublishLiveLink);
 
@@ -181,7 +181,8 @@ void FAppleARKitLiveLinkSource::PublishBlendShapes(FName SubjectName, double Tim
 	{
 		static FLiveLinkFrameData LiveLinkFrame;
 
-		LiveLinkFrame.WorldTime = Timestamp;
+		LiveLinkFrame.WorldTime = FPlatformTime::Seconds();
+		LiveLinkFrame.MetaData.SceneTime = FQualifiedFrameTime(Timecode, FFrameRate(FrameRate, 1));
 		
 		TArray<FLiveLinkCurveElement>& BlendShapes = LiveLinkFrame.CurveElements;
 
@@ -210,10 +211,12 @@ PRAGMA_ENABLE_DEPRECATION_WARNINGS
 // 2 = ARKit 2.0 extra blendshapes
 // 3 = Removed the timestamp to derive locally
 // 4 = Added the device id to stream so we can tell devices apart
-const uint8 BLEND_SHAPE_PACKET_VER = 4;
+// 5 = Added timecode tracking of the curve data for synchronizing across multiple devices, media source, etc.
+const uint8 BLEND_SHAPE_PACKET_VER = 5;
 
-const uint32 MAX_BLEND_SHAPE_PACKET_SIZE = sizeof(BLEND_SHAPE_PACKET_VER) + sizeof(uint32) + sizeof(uint8) + (sizeof(float) * (uint64)EARFaceBlendShape::MAX) + (sizeof(TCHAR) * 256) + (sizeof(TCHAR) * 256);
-const uint32 MIN_BLEND_SHAPE_PACKET_SIZE = sizeof(BLEND_SHAPE_PACKET_VER) + sizeof(uint32) + sizeof(uint8) + (sizeof(float) * (uint64)EARFaceBlendShape::MAX) + sizeof(TCHAR) + sizeof(TCHAR);
+//																			Timecode			FrameRate		 BlendShapeCount Blendshapes										SubjectName				DeviceID
+const uint32 MAX_BLEND_SHAPE_PACKET_SIZE = sizeof(BLEND_SHAPE_PACKET_VER) + sizeof(FTimecode) + sizeof(uint32) + sizeof(uint8) + (sizeof(float) * (uint64)EARFaceBlendShape::MAX) + (sizeof(TCHAR) * 256) + (sizeof(TCHAR) * 256);
+const uint32 MIN_BLEND_SHAPE_PACKET_SIZE = sizeof(BLEND_SHAPE_PACKET_VER) + sizeof(FTimecode) + sizeof(uint32) + sizeof(uint8) + (sizeof(float) * (uint64)EARFaceBlendShape::MAX) + sizeof(TCHAR) + sizeof(TCHAR);
 
 FAppleARKitLiveLinkRemotePublisher::FAppleARKitLiveLinkRemotePublisher(const FString& InRemoteIp) :
 	RemoteIp(InRemoteIp),
@@ -273,7 +276,7 @@ TSharedRef<FInternetAddr> FAppleARKitLiveLinkRemotePublisher::GetSendAddress()
 	return SendAddr;
 }
 
-void FAppleARKitLiveLinkRemotePublisher::PublishBlendShapes(FName SubjectName, double Timestamp, uint32 FrameNumber, const FARBlendShapeMap& FaceBlendShapes, FName DeviceId)
+void FAppleARKitLiveLinkRemotePublisher::PublishBlendShapes(FName SubjectName, const FTimecode& Timecode, uint32 FrameRate, const FARBlendShapeMap& FaceBlendShapes, FName DeviceId)
 {
 	if (SendSocket != nullptr)
 	{
@@ -283,7 +286,8 @@ void FAppleARKitLiveLinkRemotePublisher::PublishBlendShapes(FName SubjectName, d
 		SendBuffer << BLEND_SHAPE_PACKET_VER;
 		SendBuffer << DeviceId;
 		SendBuffer << SubjectName;
-		SendBuffer << FrameNumber;
+		SendBuffer << Timecode;
+		SendBuffer << FrameRate;
 		uint8 BlendShapeCount = (uint8)EARFaceBlendShape::MAX;
 		check(FaceBlendShapes.Num() == BlendShapeCount);
 		SendBuffer << BlendShapeCount;
@@ -383,9 +387,10 @@ void FAppleARKitLiveLinkRemoteListener::Tick(float DeltaTime)
 
 			uint8 PacketVer = 0;
 			FName SubjectName;
-			uint32 FrameNumber = 0;
 			uint8 BlendShapeCount = (uint8)EARFaceBlendShape::MAX;
 			FName DeviceId;
+			FTimecode Timecode;
+			uint32 FrameRate = 60;
 
 			FNboSerializeFromBuffer FromBuffer(RecvBuffer.GetData(), BytesRead);
 
@@ -397,7 +402,8 @@ void FAppleARKitLiveLinkRemoteListener::Tick(float DeltaTime)
 			}
 			FromBuffer >> DeviceId;
 			FromBuffer >> SubjectName;
-			FromBuffer >> FrameNumber;
+			FromBuffer >> Timecode;
+			FromBuffer >> FrameRate;
 			FromBuffer >> BlendShapeCount;
 			if (FromBuffer.HasOverflow() || BlendShapeCount != (uint8)EARFaceBlendShape::MAX)
 			{
@@ -418,7 +424,7 @@ void FAppleARKitLiveLinkRemoteListener::Tick(float DeltaTime)
 				InitLiveLinkSource();
 				if (Source.IsValid())
 				{
-					Source->PublishBlendShapes(SubjectName, FPlatformTime::Seconds(), FrameNumber, BlendShapes, DeviceId);
+					Source->PublishBlendShapes(SubjectName, Timecode, FrameRate, BlendShapes, DeviceId);
 				}
 			}
 			else
@@ -446,6 +452,11 @@ FAppleARKitLiveLinkFileWriter::~FAppleARKitLiveLinkFileWriter()
 	}
 }
 
+void FAppleARKitLiveLinkFileWriter::SetTimecodeProvider(UTimecodeProvider* InTimecodeProvider)
+{
+	TimecodeProvider = InTimecodeProvider;
+}
+
 void FAppleARKitLiveLinkFileWriter::SaveFileData()
 {
 	FString SaveData;
@@ -459,21 +470,22 @@ void FAppleARKitLiveLinkFileWriter::SaveFileData()
 
 FString FAppleARKitLiveLinkFileWriter::GenerateFilePath()
 {
+	FTimecode Timecode = TimecodeProvider->GetTimecode();
 	FDateTime DateTime = FDateTime::UtcNow();
 	const FString UserDir = FPlatformProcess::UserDir();
 	const FString DeviceNameString = DeviceName.ToString();
 	return FString::Printf(TEXT("%sFaceTracking/%s_%d-%d-%d-%d-%d-%d-%d%s"), *UserDir, *DeviceNameString,
-		DateTime.GetYear(), DateTime.GetMonth(), DateTime.GetDay(), DateTime.GetHour(), DateTime.GetMinute(), DateTime.GetSecond(), DateTime.GetMillisecond(),
+		DateTime.GetYear(), DateTime.GetMonth(), DateTime.GetDay(), Timecode.Hours, Timecode.Minutes, Timecode.Seconds, Timecode.Frames,
 		*FileExtension);
 }
 
-void FAppleARKitLiveLinkFileWriter::PublishBlendShapes(FName SubjectName, double Timestamp, uint32 FrameNumber, const FARBlendShapeMap& FaceBlendShapes, FName DeviceId)
+void FAppleARKitLiveLinkFileWriter::PublishBlendShapes(FName SubjectName, const FTimecode& Timecode, uint32 FrameRate, const FARBlendShapeMap& FaceBlendShapes, FName DeviceId)
 {
 	FScopeLock ScopeLock(&CriticalSection);
 
 	DeviceName = DeviceId;
 	// Add to the array for long running save
-	new(FrameHistory) FFaceTrackingFrame(Timestamp, FrameNumber, FaceBlendShapes);
+	new(FrameHistory) FFaceTrackingFrame(Timecode, FrameRate, FaceBlendShapes);
 
 	if (bSavePerFrameOrOnDemand)
 	{
@@ -499,7 +511,7 @@ FAppleARKitLiveLinkFileWriterCsv::FAppleARKitLiveLinkFileWriterCsv()
 	// Touching UObjects, so needs to be game thread
 	check(IsInGameThread());
 
-	CsvFrameHeader = TEXT("Timestamp, FrameNumber");
+	CsvFrameHeader = TEXT("Timecode, FrameRate");
 	const UEnum *EnumPtr = FindObject<UEnum>(ANY_PACKAGE, TEXT("EARFaceBlendShape"), true);
 	if (EnumPtr != nullptr)
 	{
@@ -516,7 +528,9 @@ FAppleARKitLiveLinkFileWriterCsv::FAppleARKitLiveLinkFileWriterCsv()
 
 FString FAppleARKitLiveLinkFileWriterCsv::BuildCsvRow(const FFaceTrackingFrame& Frame)
 {
-	FString SaveData = FString::Printf(TEXT("%f, %d"), Frame.Timestamp, Frame.FrameNumber);
+	FString SaveData = FString::Printf(TEXT("%d:%d:%d:%d, %d"),
+			Frame.Timecode.Hours, Frame.Timecode.Minutes, Frame.Timecode.Seconds, Frame.Timecode.Frames,
+			Frame.FrameRate);
 	// Add all of the blend shapes on
 	for (int32 Shape = 0; Shape < (int32)EARFaceBlendShape::MAX; Shape++)
 	{
@@ -561,7 +575,9 @@ FAppleARKitLiveLinkFileWriterJson::FAppleARKitLiveLinkFileWriterJson()
 
 FString FAppleARKitLiveLinkFileWriterJson::BuildJsonRow(const FFaceTrackingFrame& Frame)
 {
-	FString SaveData = FString::Printf(TEXT("\t{\r\n\t\t\"Timestamp\" : %f,\r\n\t\t\"FrameNumber\" : %d,\r\n"), Frame.Timestamp, Frame.FrameNumber);
+	FString SaveData = FString::Printf(TEXT("\t{\r\n\t\t\"TimeCode\" :\r\n\t\t{\r\n\t\t\t\"Hours\" : %d,\r\n\t\t\t\"Minutes\" : %d,\r\n\t\t\t\"Seconds\" : %d,\r\n\t\t\t\"Frames\" : %d\r\n\t\t},\r\n"),
+			Frame.Timecode.Hours, Frame.Timecode.Minutes, Frame.Timecode.Seconds, Frame.Timecode.Frames);
+	SaveData += FString::Printf(TEXT("\t\t\"FrameRate\" : %d,\r\n"), Frame.FrameRate);
 	bool bNeedsComma = false;
 	// Add all of the blend shapes on
 	for (int32 Shape = 0; Shape < (int32)EARFaceBlendShape::MAX; Shape++)
diff --git a/Engine/Plugins/Runtime/AR/Apple/AppleARKitFaceSupport/Source/AppleARKitFaceSupport/Private/AppleARKitLiveLinkSource.h b/Engine/Plugins/Runtime/AR/Apple/AppleARKitFaceSupport/Source/AppleARKitFaceSupport/Private/AppleARKitLiveLinkSource.h
--- a/Engine/Plugins/Runtime/AR/Apple/AppleARKitFaceSupport/Source/AppleARKitFaceSupport/Private/AppleARKitLiveLinkSource.h
+++ b/Engine/Plugins/Runtime/AR/Apple/AppleARKitFaceSupport/Source/AppleARKitFaceSupport/Private/AppleARKitLiveLinkSource.h
@@ -342,6 +342,20 @@ public:
 		return Ar;
 	}
 
+	/**
+	 * Adds a FTimecode to the buffer
+	 */
+	friend inline FNboSerializeToBuffer& operator<<(FNboSerializeToBuffer& Ar,const FTimecode& Timecode)
+	{
+		Ar << Timecode.Hours;
+		Ar << Timecode.Minutes;
+		Ar << Timecode.Seconds;
+		Ar << Timecode.Frames;
+		uint8 Drop = Timecode.bDropFrameFormat ? 1 : 0;
+		Ar << Drop;
+		return Ar;
+	}
+
 	/**
 	 * Writes a blob of data to the buffer
 	 *
@@ -615,6 +629,21 @@ public:
 		return Ar;
 	}
 
+	/**
+	 * Reads a FTimecode from the buffer
+	 */
+	friend inline FNboSerializeFromBuffer& operator>>(FNboSerializeFromBuffer& Ar,FTimecode& Timecode)
+	{
+		Ar >> Timecode.Hours;
+		Ar >> Timecode.Minutes;
+		Ar >> Timecode.Seconds;
+		Ar >> Timecode.Frames;
+		uint8 Drop = 0;
+		Ar >> Drop;
+		Timecode.bDropFrameFormat = Drop > 0 ? true : false;
+		return Ar;
+	}
+
 	/**
 	 * Reads a blob of data from the buffer
 	 *
@@ -733,7 +762,7 @@ public:
 	virtual ~FAppleARKitLiveLinkRemotePublisher();
 
 	// IARKitBlendShapePublisher interface
-	virtual void PublishBlendShapes(FName SubjectName, double Timestamp, uint32 FrameNumber, const FARBlendShapeMap& FaceBlendShapes, FName DeviceId = NAME_None) override;
+	virtual void PublishBlendShapes(FName SubjectName, const FTimecode& Timecode, uint32 FrameRate, const FARBlendShapeMap& FaceBlendShapes, FName DeviceId = NAME_None) override;
 	// End IARKitBlendShapePublisher
 
 	bool InitSendSocket();
@@ -756,26 +785,27 @@ class FAppleARKitLiveLinkFileWriter :
 {
 public:
 	virtual ~FAppleARKitLiveLinkFileWriter();
+	virtual void SetTimecodeProvider(UTimecodeProvider* InTimecodeProvider = nullptr) override;
 
 protected:
 	FAppleARKitLiveLinkFileWriter(const TCHAR* InFileExtension);
 
 	struct FFaceTrackingFrame
 	{
-		FFaceTrackingFrame(double InTimestamp, uint32 InFrameNumber, const FARBlendShapeMap& InBlendShapes)
-			: Timestamp(InTimestamp)
-			, FrameNumber(InFrameNumber)
+		FFaceTrackingFrame(const FTimecode& InTimecode, uint32 InFrameRate, const FARBlendShapeMap& InBlendShapes)
+			: Timecode(InTimecode)
+			, FrameRate(InFrameRate)
 			, BlendShapes(InBlendShapes)
 		{
 		}
 
-		double Timestamp;
-		uint32 FrameNumber;
+		FTimecode Timecode;
+		uint32 FrameRate;
 		FARBlendShapeMap BlendShapes;
 	};
 
 	// IARKitBlendShapePublisher interface
-	virtual void PublishBlendShapes(FName SubjectName, double Timestamp, uint32 FrameNumber, const FARBlendShapeMap& FaceBlendShapes, FName DeviceId = NAME_None) override;
+	virtual void PublishBlendShapes(FName SubjectName, const FTimecode& Timecode, uint32 FrameRate, const FARBlendShapeMap& FaceBlendShapes, FName DeviceId = NAME_None) override;
 	// End IARKitBlendShapePublisher
 
 	//~ FSelfRegisteringExec
@@ -797,6 +827,12 @@ protected:
 	FName DeviceName;
 	/** If true, it writes a file each frame. Otherwise, it writes the data upon demand or exit */
 	bool bSavePerFrameOrOnDemand;
+	/**
+	 * The time code provider to use when tagging file time stamps
+	 * Note: this requires the FAppleARKitSystem object to mark it in use so GC doesn't destroy it. Normally it would
+	 * implement the FGCObject interface but this gets created before UObjects are init-ed so not possible
+	 * */
+	UTimecodeProvider* TimecodeProvider;
 };
 
 class FAppleARKitLiveLinkFileWriterJson :
@@ -852,7 +888,7 @@ private:
 	// End ILiveLinkSource
 
 	// IARKitBlendShapePublisher interface
-	virtual void PublishBlendShapes(FName SubjectName, double Timestamp, uint32 FrameNumber, const FARBlendShapeMap& FaceBlendShapes, FName DeviceId = NAME_None) override;
+	virtual void PublishBlendShapes(FName SubjectName, const FTimecode& Timecode, uint32 FrameRate, const FARBlendShapeMap& FaceBlendShapes, FName DeviceId = NAME_None) override;
 	// End IARKitBlendShapePublisher
 
 	/** The local client to push data updates to */
diff --git a/Engine/Plugins/Runtime/AR/Apple/AppleARKitFaceSupport/Source/AppleARKitFaceSupport/Public/AppleARKitLiveLinkSourceFactory.h b/Engine/Plugins/Runtime/AR/Apple/AppleARKitFaceSupport/Source/AppleARKitFaceSupport/Public/AppleARKitLiveLinkSourceFactory.h
--- a/Engine/Plugins/Runtime/AR/Apple/AppleARKitFaceSupport/Source/AppleARKitFaceSupport/Public/AppleARKitLiveLinkSourceFactory.h
+++ b/Engine/Plugins/Runtime/AR/Apple/AppleARKitFaceSupport/Source/AppleARKitFaceSupport/Public/AppleARKitLiveLinkSourceFactory.h
@@ -6,12 +6,14 @@
 #include "ARTrackable.h"
 #include "ILiveLinkSource.h"
 
+class UTimecodeProvider;
 
 /** Interface that publishes face ar blend shape information */
 class APPLEARKITFACESUPPORT_API IARKitBlendShapePublisher
 {
 public:
-	virtual void PublishBlendShapes(FName SubjectName, double Timestamp, uint32 FrameNumber, const FARBlendShapeMap& BlendShapes, FName DeviceID = NAME_None) = 0;
+	virtual void SetTimecodeProvider(UTimecodeProvider* InTimecodeProvider = nullptr) {}
+	virtual void PublishBlendShapes(FName SubjectName, const FTimecode& Timecode, uint32 FrameRate, const FARBlendShapeMap& BlendShapes, FName DeviceID = NAME_None) = 0;
 };
 
 /** Interface that publishes face ar blend shape information via LiveLink */
diff --git a/Engine/Source/Runtime/Core/Public/Misc/QualifiedFrameTime.h b/Engine/Source/Runtime/Core/Public/Misc/QualifiedFrameTime.h
--- a/Engine/Source/Runtime/Core/Public/Misc/QualifiedFrameTime.h
+++ b/Engine/Source/Runtime/Core/Public/Misc/QualifiedFrameTime.h
@@ -4,6 +4,7 @@
 
 #include "Misc/FrameRate.h"
 #include "Misc/FrameTime.h"
+#include "Misc/Timecode.h"
 
 /**
  * A frame time qualified by a frame rate context
@@ -25,6 +26,17 @@ struct FQualifiedFrameTime
 		: Time(InTime), Rate(InRate)
 	{}
 
+	/**
+	 * User construction from a timecode and its frame rate
+	 */
+	FQualifiedFrameTime(const FTimecode& InTimecode, const FFrameRate& InRate)
+		: Time(InTimecode.Frames)
+		, Rate(InRate)
+	{
+		// Check for cases where the user passed invalid data
+		check(InTimecode.Frames < InRate.Numerator);
+	}
+
 public:
 
 	/**
diff --git a/Engine/Source/Runtime/Core/Public/Misc/Timecode.h b/Engine/Source/Runtime/Core/Public/Misc/Timecode.h
--- a/Engine/Source/Runtime/Core/Public/Misc/Timecode.h
+++ b/Engine/Source/Runtime/Core/Public/Misc/Timecode.h
@@ -38,6 +38,21 @@ struct FTimecode
 		, bDropFrameFormat(InbDropFrame)
 	{}
 
+	FTimecode(double InSeconds, const FFrameRate& FrameRate)
+		: bDropFrameFormat(false)
+	{
+		int32 WholeSeconds = InSeconds;
+		const float FractionalSeconds = InSeconds - WholeSeconds;
+
+		Frames = FMath::FloorToInt(FrameRate.Numerator * FractionalSeconds);
+
+		Seconds = WholeSeconds % 60;
+		WholeSeconds /= 60;
+		Minutes = WholeSeconds % 60;
+		WholeSeconds /= 60;
+		Hours = WholeSeconds % 24;
+	}
+
 public:
 
 	/**
@@ -219,4 +234,4 @@ public:
 
 	/** If true, this Timecode represents a Drop Frame timecode used to account for fractional frame rates in NTSC play rates. */
 	bool bDropFrameFormat;
-};
\ No newline at end of file
+};
