Added some console commands to control face ar data logging
- "FaceAR StartFileWriting" to start writing
- "FaceAR StopFileWriting" to stop writing
- "FaceAR SavePerFrame" to write a file each frame when logging is enabled
- "FaceAR SaveOnDemand" to write a file upon close of the app or upon user command

diff --git a/Engine/Plugins/Runtime/AR/Apple/AppleARKit/Source/AppleARKit/Public/AppleARKitSettings.h b/Engine/Plugins/Runtime/AR/Apple/AppleARKit/Source/AppleARKit/Public/AppleARKitSettings.h
--- a/Engine/Plugins/Runtime/AR/Apple/AppleARKit/Source/AppleARKit/Public/AppleARKitSettings.h
+++ b/Engine/Plugins/Runtime/AR/Apple/AppleARKit/Source/AppleARKit/Public/AppleARKitSettings.h
@@ -33,6 +33,7 @@ class APPLEARKIT_API UAppleARKitSettings :
 public:
 	UAppleARKitSettings()
 		: bEnableLiveLinkForFaceTracking(false)
+		, bFaceTrackingLogData(false)
 		, bFaceTrackingWriteEachFrame(false)
 		, FaceTrackingFileWriterType(EARFaceTrackingFileWriterType::None)
 		, bShouldWriteCameraImagePerFrame(false)
@@ -52,6 +53,10 @@ public:
 	UPROPERTY(Config, EditAnywhere, Category="AR Settings")
 	bool bEnableLiveLinkForFaceTracking;
 
+	/** Whether file writing is enabled at all or not */
+	UPROPERTY(Config, EditAnywhere, Category="AR Settings")
+	bool bFaceTrackingLogData;
+
 	/** Whether to publish each frame or when the "FaceAR WriteCurveFile */
 	UPROPERTY(Config, EditAnywhere, Category="AR Settings")
 	bool bFaceTrackingWriteEachFrame;
diff --git a/Engine/Plugins/Runtime/AR/Apple/AppleARKitFaceSupport/Source/AppleARKitFaceSupport/Private/AppleARKitLiveLinkSource.cpp b/Engine/Plugins/Runtime/AR/Apple/AppleARKitFaceSupport/Source/AppleARKitFaceSupport/Private/AppleARKitLiveLinkSource.cpp
--- a/Engine/Plugins/Runtime/AR/Apple/AppleARKitFaceSupport/Source/AppleARKitFaceSupport/Private/AppleARKitLiveLinkSource.cpp
+++ b/Engine/Plugins/Runtime/AR/Apple/AppleARKitFaceSupport/Source/AppleARKitFaceSupport/Private/AppleARKitLiveLinkSource.cpp
@@ -441,6 +441,7 @@ FAppleARKitLiveLinkFileWriter::FAppleARKitLiveLinkFileWriter(const TCHAR* InFile
 {
 	// Read the config values for this
 	GConfig->GetBool(TEXT("/Script/AppleARKit.AppleARKitSettings"), TEXT("bFaceTrackingWriteEachFrame"), bSavePerFrameOrOnDemand, GEngineIni);
+	GConfig->GetBool(TEXT("/Script/AppleARKit.AppleARKitSettings"), TEXT("bFaceTrackingLogData"), bLogData, GEngineIni);
 }
 
 FAppleARKitLiveLinkFileWriter::~FAppleARKitLiveLinkFileWriter()
@@ -483,6 +484,11 @@ void FAppleARKitLiveLinkFileWriter::PublishBlendShapes(FName SubjectName, const
 {
 	FScopeLock ScopeLock(&CriticalSection);
 
+	if (!bLogData)
+	{
+		return;
+	}
+
 	DeviceName = DeviceId;
 	// Add to the array for long running save
 	new(FrameHistory) FFaceTrackingFrame(Timecode, FrameRate, FaceBlendShapes);
@@ -495,12 +501,40 @@ void FAppleARKitLiveLinkFileWriter::PublishBlendShapes(FName SubjectName, const
 
 bool FAppleARKitLiveLinkFileWriter::Exec(UWorld*, const TCHAR* Cmd, FOutputDevice& Ar)
 {
-	if (FParse::Command(&Cmd, TEXT("FaceAR")) &&
-		FParse::Command(&Cmd, TEXT("WriteCurveFile")))
+	if (FParse::Command(&Cmd, TEXT("FaceAR")))
 	{
-		FScopeLock ScopeLock(&CriticalSection);
-		SaveFileData();
-		return true;
+		if (FParse::Command(&Cmd, TEXT("WriteCurveFile")))
+		{
+			FScopeLock ScopeLock(&CriticalSection);
+			SaveFileData();
+			return true;
+		}
+		else if (FParse::Command(&Cmd, TEXT("StartFileWriting")))
+		{
+			FScopeLock ScopeLock(&CriticalSection);
+			FrameHistory.Empty();
+			bLogData = true;
+			return true;
+		}
+		else if (FParse::Command(&Cmd, TEXT("StopFileWriting")))
+		{
+			FScopeLock ScopeLock(&CriticalSection);
+			FrameHistory.Empty();
+			bLogData = false;
+			return true;
+		}
+		else if (FParse::Command(&Cmd, TEXT("SavePerFrame")))
+		{
+			FScopeLock ScopeLock(&CriticalSection);
+			bSavePerFrameOrOnDemand = true;
+			return true;
+		}
+		else if (FParse::Command(&Cmd, TEXT("SaveOnDemand")))
+		{
+			FScopeLock ScopeLock(&CriticalSection);
+			bSavePerFrameOrOnDemand = false;
+			return true;
+		}
 	}
 	return false;
 }
diff --git a/Engine/Plugins/Runtime/AR/Apple/AppleARKitFaceSupport/Source/AppleARKitFaceSupport/Private/AppleARKitLiveLinkSource.h b/Engine/Plugins/Runtime/AR/Apple/AppleARKitFaceSupport/Source/AppleARKitFaceSupport/Private/AppleARKitLiveLinkSource.h
--- a/Engine/Plugins/Runtime/AR/Apple/AppleARKitFaceSupport/Source/AppleARKitFaceSupport/Private/AppleARKitLiveLinkSource.h
+++ b/Engine/Plugins/Runtime/AR/Apple/AppleARKitFaceSupport/Source/AppleARKitFaceSupport/Private/AppleARKitLiveLinkSource.h
@@ -825,6 +825,8 @@ protected:
 	TArray<FFaceTrackingFrame> FrameHistory;
 	/** The name to use when generating the file to save to */
 	FName DeviceName;
+	/** If false, the calls to PublishCurves are ignored */
+	bool bLogData;
 	/** If true, it writes a file each frame. Otherwise, it writes the data upon demand or exit */
 	bool bSavePerFrameOrOnDemand;
 	/**
