Add support for fractional timecode to ARKit face protocol.

diff --git a/Engine/Plugins/Runtime/AR/Apple/AppleARKitFaceSupport/Source/AppleARKitFaceSupport/Private/AppleARKitFaceMeshComponent.cpp b/Engine/Plugins/Runtime/AR/Apple/AppleARKitFaceSupport/Source/AppleARKitFaceSupport/Private/AppleARKitFaceMeshComponent.cpp
--- a/Engine/Plugins/Runtime/AR/Apple/AppleARKitFaceSupport/Source/AppleARKitFaceSupport/Private/AppleARKitFaceMeshComponent.cpp
+++ b/Engine/Plugins/Runtime/AR/Apple/AppleARKitFaceSupport/Source/AppleARKitFaceSupport/Private/AppleARKitFaceMeshComponent.cpp
@@ -386,14 +386,23 @@ void UAppleARKitFaceMeshComponent::OnRep_RemoteCurves()
 		{
 			BlendShapes.Add(Curve.GetBlendShape(), Curve.GetAmountAsFloat());
 		}
+
 		// We have to manage this since the ar system isn't updating this locally for us
 		LastUpdateTimestamp = FPlatformTime::Seconds();
-		FTimecode Timecode = FApp::GetTimecode();
-		FFrameRate FrameRate = FApp::GetTimecodeFrameRate();
+
 		// If we are publishing to LiveLink, then push the update out
 		if (LiveLinkSource.IsValid())
 		{
-			LiveLinkSource->PublishBlendShapes(LiveLinkSubjectName, Timecode, FrameRate.Numerator, BlendShapes);
+			TOptional<FQualifiedFrameTime> FrameTimeOpt = FApp::GetCurrentFrameTime();
+
+			if (FrameTimeOpt.IsSet())
+			{
+				LiveLinkSource->PublishBlendShapes(LiveLinkSubjectName, *FrameTimeOpt, BlendShapes);
+			}
+			else
+			{
+				LiveLinkSource->PublishBlendShapes(LiveLinkSubjectName, FQualifiedFrameTime(), BlendShapes);
+			}
 		}
 	}
 }
diff --git a/Engine/Plugins/Runtime/AR/Apple/AppleARKitFaceSupport/Source/AppleARKitFaceSupport/Private/AppleARKitFaceSupportImpl.cpp b/Engine/Plugins/Runtime/AR/Apple/AppleARKitFaceSupport/Source/AppleARKitFaceSupport/Private/AppleARKitFaceSupportImpl.cpp
--- a/Engine/Plugins/Runtime/AR/Apple/AppleARKitFaceSupport/Source/AppleARKitFaceSupport/Private/AppleARKitFaceSupportImpl.cpp
+++ b/Engine/Plugins/Runtime/AR/Apple/AppleARKitFaceSupport/Source/AppleARKitFaceSupport/Private/AppleARKitFaceSupportImpl.cpp
@@ -10,6 +10,7 @@
 #include "AppleARKitFaceSupportModule.h"
 #include "Async/Async.h"
 #include "Engine/TimecodeProvider.h"
+#include "Math/UnrealMathUtility.h"
 
 DECLARE_CYCLE_STAT(TEXT("Conversion"), STAT_FaceAR_Conversion, STATGROUP_FaceAR);
 
@@ -122,7 +123,6 @@ void FAppleARKitFaceSupport::InitRealtimeProviders()
 #endif
 	}
 }
-
 #if SUPPORTS_ARKIT_1_0
 
 ARConfiguration* FAppleARKitFaceSupport::ToARConfiguration(UARSessionConfig* SessionConfig, UTimecodeProvider* InProvider)
@@ -174,9 +174,22 @@ TArray<TSharedPtr<FAppleARKitAnchorData>> FAppleARKitFaceSupport::MakeAnchorData
 	const FTimecode& Timecode = TimecodeProvider->GetTimecode();
 	const FFrameRate& FrameRate = TimecodeProvider->GetFrameRate();
 
+	int32 Rate = 1;
+
+	if (FrameRate.IsValid()) // avoids divide by zero when calling FrameRate.AsDecimal()
+	{
+		Rate = int32(FMath::RoundHalfFromZero(FrameRate.AsDecimal()));
+
+		// Just in case it rounded to zero (or negative).
+		if (Rate < 1)
+		{
+			Rate = 1;
+		}
+	}
+
 	for (ARAnchor* Anchor in Anchors)
 	{
-		TSharedPtr<FAppleARKitAnchorData> AnchorData = ::MakeAnchorData(bFaceMirrored, Anchor, AdjustBy, UpdateSetting, Timecode, FrameRate.Numerator);
+		TSharedPtr<FAppleARKitAnchorData> AnchorData = ::MakeAnchorData(bFaceMirrored, Anchor, AdjustBy, UpdateSetting, Timecode, uint32(Rate));
 		if (AnchorData.IsValid())
 		{
 			AnchorList.Add(AnchorData);
@@ -194,18 +207,17 @@ void FAppleARKitFaceSupport::ProcessRealTimePublishers(TSharedPtr<FAppleARKitAnc
     TSharedPtr<FAppleARKitAnchorData> AsyncAnchorCopy = MakeShared<FAppleARKitAnchorData>(*AnchorData);
 	AsyncTask(ENamedThreads::AnyBackgroundThreadNormalTask, [this, AsyncAnchorCopy]()
 	{
-		FTimecode Timecode = AsyncAnchorCopy->Timecode;
-		uint32 FrameRate = AsyncAnchorCopy->FrameRate;
+		const FQualifiedFrameTime FrameTime(AsyncAnchorCopy->Timecode, FFrameRate(AsyncAnchorCopy->FrameRate, 1));
 		const FARBlendShapeMap& BlendShapes = AsyncAnchorCopy->BlendShapes;
 
 		if (RemoteLiveLinkPublisher.IsValid())
 		{
-			RemoteLiveLinkPublisher->PublishBlendShapes(FaceTrackingLiveLinkSubjectName, Timecode, FrameRate, BlendShapes, LocalDeviceId);
+			RemoteLiveLinkPublisher->PublishBlendShapes(FaceTrackingLiveLinkSubjectName, FrameTime, BlendShapes, LocalDeviceId);
 		}
 
 		if (LiveLinkFileWriter.IsValid())
 		{
-			LiveLinkFileWriter->PublishBlendShapes(FaceTrackingLiveLinkSubjectName, Timecode, FrameRate, BlendShapes, LocalDeviceId);
+			LiveLinkFileWriter->PublishBlendShapes(FaceTrackingLiveLinkSubjectName, FrameTime, BlendShapes, LocalDeviceId);
 		}
 	});
 }
@@ -230,8 +242,10 @@ void FAppleARKitFaceSupport::PublishLiveLinkData(TSharedPtr<FAppleARKitAnchorDat
 
 	if (LiveLinkSource.IsValid())
 	{
+		const FQualifiedFrameTime FrameTime(Anchor->Timecode, FFrameRate(Anchor->FrameRate, 1));
+
 		FaceTrackingLiveLinkSubjectName = GetMutableDefault<UAppleARKitSettings>()->GetLiveLinkSubjectName();
-        LiveLinkSource->PublishBlendShapes(FaceTrackingLiveLinkSubjectName, Anchor->Timecode, Anchor->FrameRate, Anchor->BlendShapes, LocalDeviceId);
+        LiveLinkSource->PublishBlendShapes(FaceTrackingLiveLinkSubjectName, FrameTime, Anchor->BlendShapes, LocalDeviceId);
 	}
 }
 
diff --git a/Engine/Plugins/Runtime/AR/Apple/AppleARKitFaceSupport/Source/AppleARKitFaceSupport/Private/AppleARKitLiveLinkSource.cpp b/Engine/Plugins/Runtime/AR/Apple/AppleARKitFaceSupport/Source/AppleARKitFaceSupport/Private/AppleARKitLiveLinkSource.cpp
--- a/Engine/Plugins/Runtime/AR/Apple/AppleARKitFaceSupport/Source/AppleARKitFaceSupport/Private/AppleARKitLiveLinkSource.cpp
+++ b/Engine/Plugins/Runtime/AR/Apple/AppleARKitFaceSupport/Source/AppleARKitFaceSupport/Private/AppleARKitLiveLinkSource.cpp
@@ -155,7 +155,7 @@ static FName ParseEnumName(FName EnumName)
 	return FName(*EnumString.Right(EnumString.Len() - BlendShapeEnumNameLength));
 }
 
-void FAppleARKitLiveLinkSource::PublishBlendShapes(FName SubjectName, const FTimecode& Timecode, uint32 FrameRate, const FARBlendShapeMap& FaceBlendShapes, FName DeviceId)
+void FAppleARKitLiveLinkSource::PublishBlendShapes(FName SubjectName, const FQualifiedFrameTime& FrameTime, const FARBlendShapeMap& FaceBlendShapes, FName DeviceId)
 {
 	SCOPE_CYCLE_COUNTER(STAT_FaceAR_Local_PublishLiveLink);
 
@@ -192,7 +192,7 @@ void FAppleARKitLiveLinkSource::PublishBlendShapes(FName SubjectName, const FTim
 	FLiveLinkFrameDataStruct FrameDataStruct(FLiveLinkBaseFrameData::StaticStruct());
 	FLiveLinkBaseFrameData* FrameData = FrameDataStruct.Cast<FLiveLinkBaseFrameData>();
 	FrameData->WorldTime = FPlatformTime::Seconds();
-	FrameData->MetaData.SceneTime = FQualifiedFrameTime(Timecode, FFrameRate(FrameRate, 1));
+	FrameData->MetaData.SceneTime = FrameTime;
 	FrameData->PropertyValues.Reserve((int32)EARFaceBlendShape::MAX);
 	
 	// Iterate through all of the blend shapes copying them into the LiveLink data type
@@ -250,11 +250,12 @@ void FAppleARKitLiveLinkSource::UpdateStaticData(FName SubjectName, const FARBle
 // 3 = Removed the timestamp to derive locally
 // 4 = Added the device id to stream so we can tell devices apart
 // 5 = Added timecode tracking of the curve data for synchronizing across multiple devices, media source, etc.
-const uint8 BLEND_SHAPE_PACKET_VER = 5;
+// 6 = Replaced FTimecode with FQualifiedFrameTime for accurate timestamping of ARKit capture, which is currently asynchronous.
+const uint8 BLEND_SHAPE_PACKET_VER = 6;
 
-//																			Timecode			FrameRate		 BlendShapeCount Blendshapes										SubjectName				DeviceID
-const uint32 MAX_BLEND_SHAPE_PACKET_SIZE = sizeof(BLEND_SHAPE_PACKET_VER) + sizeof(FTimecode) + sizeof(uint32) + sizeof(uint8) + (sizeof(float) * (uint64)EARFaceBlendShape::MAX) + (sizeof(TCHAR) * 256) + (sizeof(TCHAR) * 256);
-const uint32 MIN_BLEND_SHAPE_PACKET_SIZE = sizeof(BLEND_SHAPE_PACKET_VER) + sizeof(FTimecode) + sizeof(uint32) + sizeof(uint8) + (sizeof(float) * (uint64)EARFaceBlendShape::MAX) + sizeof(TCHAR) + sizeof(TCHAR);
+//                                                                          FrameTime                     BlendShapeCount Blendshapes                                        SubjectName             DeviceID
+const uint32 MAX_BLEND_SHAPE_PACKET_SIZE = sizeof(BLEND_SHAPE_PACKET_VER) + sizeof(FQualifiedFrameTime) + sizeof(uint8) + (sizeof(float) * (uint64)EARFaceBlendShape::MAX) + (sizeof(TCHAR) * 256) + (sizeof(TCHAR) * 256);
+const uint32 MIN_BLEND_SHAPE_PACKET_SIZE = sizeof(BLEND_SHAPE_PACKET_VER) + sizeof(FQualifiedFrameTime) + sizeof(uint8) + (sizeof(float) * (uint64)EARFaceBlendShape::MAX) +  sizeof(TCHAR)        +  sizeof(TCHAR);
 
 FAppleARKitLiveLinkRemotePublisher::FAppleARKitLiveLinkRemotePublisher(const FString& InRemoteIp) :
 	RemoteIp(InRemoteIp),
@@ -314,7 +315,7 @@ TSharedRef<FInternetAddr> FAppleARKitLiveLinkRemotePublisher::GetSendAddress()
 	return SendAddr;
 }
 
-void FAppleARKitLiveLinkRemotePublisher::PublishBlendShapes(FName SubjectName, const FTimecode& Timecode, uint32 FrameRate, const FARBlendShapeMap& FaceBlendShapes, FName DeviceId)
+void FAppleARKitLiveLinkRemotePublisher::PublishBlendShapes(FName SubjectName, const FQualifiedFrameTime& FrameTime, const FARBlendShapeMap& FaceBlendShapes, FName DeviceId)
 {
 	if (SendSocket != nullptr)
 	{
@@ -324,8 +325,7 @@ void FAppleARKitLiveLinkRemotePublisher::PublishBlendShapes(FName SubjectName, c
 		SendBuffer << BLEND_SHAPE_PACKET_VER;
 		SendBuffer << DeviceId;
 		SendBuffer << SubjectName;
-		SendBuffer << Timecode;
-		SendBuffer << FrameRate;
+		SendBuffer << FrameTime;
 		uint8 BlendShapeCount = (uint8)EARFaceBlendShape::MAX;
 		check(FaceBlendShapes.Num() == BlendShapeCount);
 		SendBuffer << BlendShapeCount;
@@ -401,7 +401,7 @@ bool FAppleARKitLiveLinkRemoteListener::InitReceiveSocket()
 
 void FAppleARKitLiveLinkRemoteListener::InitLiveLinkSource()
 {
-	if (!Source.IsValid())
+	if (!Source.IsValid() || !Source->IsSourceStillValid())
 	{
 		Source = FAppleARKitLiveLinkSourceFactory::CreateLiveLinkSource();
 	}
@@ -431,8 +431,7 @@ void FAppleARKitLiveLinkRemoteListener::Tick(float DeltaTime)
 			FName SubjectName;
 			uint8 BlendShapeCount = (uint8)EARFaceBlendShape::MAX;
 			FName DeviceId;
-			FTimecode Timecode;
-			uint32 FrameRate = 60;
+			FQualifiedFrameTime FrameTime;
 
 			FNboSerializeFromBuffer FromBuffer(RecvBuffer.GetData(), BytesRead);
 
@@ -444,8 +443,7 @@ void FAppleARKitLiveLinkRemoteListener::Tick(float DeltaTime)
 			}
 			FromBuffer >> DeviceId;
 			FromBuffer >> SubjectName;
-			FromBuffer >> Timecode;
-			FromBuffer >> FrameRate;
+			FromBuffer >> FrameTime;
 			FromBuffer >> BlendShapeCount;
 			if (FromBuffer.HasOverflow() || BlendShapeCount != (uint8)EARFaceBlendShape::MAX)
 			{
@@ -464,9 +462,9 @@ void FAppleARKitLiveLinkRemoteListener::Tick(float DeltaTime)
 			if (!FromBuffer.HasOverflow())
 			{
 				InitLiveLinkSource();
-				if (Source.IsValid())
+				if (Source.IsValid() && Source->IsSourceStillValid())
 				{
-					Source->PublishBlendShapes(SubjectName, Timecode, FrameRate, BlendShapes, DeviceId);
+					Source->PublishBlendShapes(SubjectName, FrameTime, BlendShapes, DeviceId);
 				}
 			}
 			else
@@ -521,7 +519,7 @@ FString FAppleARKitLiveLinkFileWriter::GenerateFilePath()
 		*FileExtension);
 }
 
-void FAppleARKitLiveLinkFileWriter::PublishBlendShapes(FName SubjectName, const FTimecode& Timecode, uint32 FrameRate, const FARBlendShapeMap& FaceBlendShapes, FName DeviceId)
+void FAppleARKitLiveLinkFileWriter::PublishBlendShapes(FName SubjectName, const FQualifiedFrameTime& FrameTime, const FARBlendShapeMap& FaceBlendShapes, FName DeviceId)
 {
 	FScopeLock ScopeLock(&CriticalSection);
 
@@ -532,7 +530,7 @@ void FAppleARKitLiveLinkFileWriter::PublishBlendShapes(FName SubjectName, const
 
 	DeviceName = DeviceId;
 	// Add to the array for long running save
-	new(FrameHistory) FFaceTrackingFrame(Timecode, FrameRate, FaceBlendShapes);
+	new(FrameHistory) FFaceTrackingFrame(FrameTime, FaceBlendShapes);
 
 	if (GetMutableDefault<UAppleARKitSettings>()->ShouldFaceTrackingLogPerFrame())
 	{
@@ -575,11 +573,23 @@ FAppleARKitLiveLinkFileWriterCsv::FAppleARKitLiveLinkFileWriterCsv()
 	CsvFrameHeader += TEXT("\r\n");
 }
 
+
 FString FAppleARKitLiveLinkFileWriterCsv::BuildCsvRow(const FFaceTrackingFrame& Frame)
 {
-	FString SaveData = FString::Printf(TEXT("%d:%d:%d:%d, %d"),
-		Frame.Timecode.Hours, Frame.Timecode.Minutes, Frame.Timecode.Seconds, Frame.Timecode.Frames,
-		Frame.FrameRate);
+	FTimecode Timecode;
+	double Rate = 1.0;
+	double FramesAsDecimal = 0.0;
+	
+	if (Frame.FrameTime.Rate.IsValid())
+	{
+		Timecode = FTimecode::FromFrameNumber(Frame.FrameTime.Time.GetFrame(), Frame.FrameTime.Rate);
+		Rate = Frame.FrameTime.Rate.AsDecimal();
+		FramesAsDecimal = double(Timecode.Frames) + Frame.FrameTime.Time.GetSubFrame();
+	}
+
+	FString SaveData = FString::Printf(TEXT("%02d:%02d:%02d:%02.3f, %.3f"), 
+		Timecode.Hours, Timecode.Minutes, Timecode.Seconds, FramesAsDecimal, Rate);
+
 	// Add all of the blend shapes on
 	for (int32 Shape = 0; Shape < (int32)EARFaceBlendShape::MAX; Shape++)
 	{
@@ -624,9 +634,20 @@ FAppleARKitLiveLinkFileWriterJson::FAppleARKitLiveLinkFileWriterJson()
 
 FString FAppleARKitLiveLinkFileWriterJson::BuildJsonRow(const FFaceTrackingFrame& Frame)
 {
-	FString SaveData = FString::Printf(TEXT("\t{\r\n\t\t\"TimeCode\" :\r\n\t\t{\r\n\t\t\t\"Hours\" : %d,\r\n\t\t\t\"Minutes\" : %d,\r\n\t\t\t\"Seconds\" : %d,\r\n\t\t\t\"Frames\" : %d\r\n\t\t},\r\n"),
-			Frame.Timecode.Hours, Frame.Timecode.Minutes, Frame.Timecode.Seconds, Frame.Timecode.Frames);
-	SaveData += FString::Printf(TEXT("\t\t\"FrameRate\" : %d,\r\n"), Frame.FrameRate);
+	FTimecode Timecode;
+	double Rate = 1.0;
+	double FramesAsDecimal = 0.0;
+
+	if (Frame.FrameTime.Rate.IsValid())
+	{
+		Timecode = FTimecode::FromFrameNumber(Frame.FrameTime.Time.GetFrame(), Frame.FrameTime.Rate);
+		Rate = Frame.FrameTime.Rate.AsDecimal();
+		FramesAsDecimal = double(Timecode.Frames) + Frame.FrameTime.Time.GetSubFrame();
+	}
+
+	FString SaveData = FString::Printf(TEXT("\t{\r\n\t\t\"TimeCode\" :\r\n\t\t{\r\n\t\t\t\"Hours\" : %d,\r\n\t\t\t\"Minutes\" : %d,\r\n\t\t\t\"Seconds\" : %d,\r\n\t\t\t\"Frames\" : %.3f\r\n\t\t},\r\n"),
+			Timecode.Hours, Timecode.Minutes, Timecode.Seconds, FramesAsDecimal);
+	SaveData += FString::Printf(TEXT("\t\t\"FrameRate\" : %.3f,\r\n"), Rate);
 	bool bNeedsComma = false;
 	// Add all of the blend shapes on
 	for (int32 Shape = 0; Shape < (int32)EARFaceBlendShape::MAX; Shape++)
diff --git a/Engine/Plugins/Runtime/AR/Apple/AppleARKitFaceSupport/Source/AppleARKitFaceSupport/Private/AppleARKitLiveLinkSource.h b/Engine/Plugins/Runtime/AR/Apple/AppleARKitFaceSupport/Source/AppleARKitFaceSupport/Private/AppleARKitLiveLinkSource.h
--- a/Engine/Plugins/Runtime/AR/Apple/AppleARKitFaceSupport/Source/AppleARKitFaceSupport/Private/AppleARKitLiveLinkSource.h
+++ b/Engine/Plugins/Runtime/AR/Apple/AppleARKitFaceSupport/Source/AppleARKitFaceSupport/Private/AppleARKitLiveLinkSource.h
@@ -345,16 +345,14 @@ public:
 	}
 
 	/**
-	 * Adds a FTimecode to the buffer
-	 */
-	friend inline FNboSerializeToBuffer& operator<<(FNboSerializeToBuffer& Ar,const FTimecode& Timecode)
-	{
-		Ar << Timecode.Hours;
-		Ar << Timecode.Minutes;
-		Ar << Timecode.Seconds;
-		Ar << Timecode.Frames;
-		uint8 Drop = Timecode.bDropFrameFormat ? 1 : 0;
-		Ar << Drop;
+	* Adds a FQualifiedFrameTime to the buffer
+	*/
+	friend inline FNboSerializeToBuffer& operator<<(FNboSerializeToBuffer& Ar, const FQualifiedFrameTime& FrameTime)
+	{
+		Ar << FrameTime.Time.FrameNumber.Value;
+		Ar << FrameTime.Time.GetSubFrame();
+		Ar << FrameTime.Rate.Numerator;
+		Ar << FrameTime.Rate.Denominator;
 		return Ar;
 	}
 
@@ -632,17 +630,20 @@ public:
 	}
 
 	/**
-	 * Reads a FTimecode from the buffer
+	 * Reads a FQualifiedFrameTime from the buffer
 	 */
-	friend inline FNboSerializeFromBuffer& operator>>(FNboSerializeFromBuffer& Ar,FTimecode& Timecode)
+	friend inline FNboSerializeFromBuffer& operator>>(FNboSerializeFromBuffer& Ar, FQualifiedFrameTime& FrameTime)
 	{
-		Ar >> Timecode.Hours;
-		Ar >> Timecode.Minutes;
-		Ar >> Timecode.Seconds;
-		Ar >> Timecode.Frames;
-		uint8 Drop = 0;
-		Ar >> Drop;
-		Timecode.bDropFrameFormat = Drop > 0 ? true : false;
+		int32 FrameNumber;
+		float SubFrame;
+
+		Ar >> FrameNumber;
+		Ar >> SubFrame;
+		Ar >> FrameTime.Rate.Numerator;
+		Ar >> FrameTime.Rate.Denominator;
+
+		FrameTime.Time = FFrameTime(FrameNumber, SubFrame);
+
 		return Ar;
 	}
 
@@ -764,7 +765,7 @@ public:
 	virtual ~FAppleARKitLiveLinkRemotePublisher();
 
 	// IARKitBlendShapePublisher interface
-	virtual void PublishBlendShapes(FName SubjectName, const FTimecode& Timecode, uint32 FrameRate, const FARBlendShapeMap& FaceBlendShapes, FName DeviceId = NAME_None) override;
+	virtual void PublishBlendShapes(FName SubjectName, const FQualifiedFrameTime& FrameTime, const FARBlendShapeMap& FaceBlendShapes, FName DeviceId = NAME_None) override;
 	// End IARKitBlendShapePublisher
 
 	bool InitSendSocket();
@@ -794,20 +795,18 @@ protected:
 
 	struct FFaceTrackingFrame
 	{
-		FFaceTrackingFrame(const FTimecode& InTimecode, uint32 InFrameRate, const FARBlendShapeMap& InBlendShapes)
-			: Timecode(InTimecode)
-			, FrameRate(InFrameRate)
+		FFaceTrackingFrame(const FQualifiedFrameTime& InFrameTime, const FARBlendShapeMap& InBlendShapes)
+			: FrameTime(InFrameTime)
 			, BlendShapes(InBlendShapes)
 		{
 		}
 
-		FTimecode Timecode;
-		uint32 FrameRate;
+		FQualifiedFrameTime FrameTime;
 		FARBlendShapeMap BlendShapes;
 	};
 
 	// IARKitBlendShapePublisher interface
-	virtual void PublishBlendShapes(FName SubjectName, const FTimecode& Timecode, uint32 FrameRate, const FARBlendShapeMap& FaceBlendShapes, FName DeviceId = NAME_None) override;
+	virtual void PublishBlendShapes(FName SubjectName, const FQualifiedFrameTime& FrameTime, const FARBlendShapeMap& FaceBlendShapes, FName DeviceId = NAME_None) override;
 	// End IARKitBlendShapePublisher
 
 	//~ FSelfRegisteringExec
@@ -888,7 +887,7 @@ private:
 	// End ILiveLinkSource
 
 	// IARKitBlendShapePublisher interface
-	virtual void PublishBlendShapes(FName SubjectName, const FTimecode& Timecode, uint32 FrameRate, const FARBlendShapeMap& FaceBlendShapes, FName DeviceId = NAME_None) override;
+	virtual void PublishBlendShapes(FName SubjectName, const FQualifiedFrameTime& FrameTime, const FARBlendShapeMap& FaceBlendShapes, FName DeviceId = NAME_None) override;
 	// End IARKitBlendShapePublisher
 
 	// Update static data property names from blend shape map. Should only happen once per subject name. Pushes static data to livelink to register subject
diff --git a/Engine/Plugins/Runtime/AR/Apple/AppleARKitFaceSupport/Source/AppleARKitFaceSupport/Public/AppleARKitLiveLinkSourceFactory.h b/Engine/Plugins/Runtime/AR/Apple/AppleARKitFaceSupport/Source/AppleARKitFaceSupport/Public/AppleARKitLiveLinkSourceFactory.h
--- a/Engine/Plugins/Runtime/AR/Apple/AppleARKitFaceSupport/Source/AppleARKitFaceSupport/Public/AppleARKitLiveLinkSourceFactory.h
+++ b/Engine/Plugins/Runtime/AR/Apple/AppleARKitFaceSupport/Source/AppleARKitFaceSupport/Public/AppleARKitLiveLinkSourceFactory.h
@@ -13,7 +13,7 @@ class APPLEARKITFACESUPPORT_API IARKitBlendShapePublisher
 {
 public:
 	virtual void SetTimecodeProvider(UTimecodeProvider* InTimecodeProvider = nullptr) {}
-	virtual void PublishBlendShapes(FName SubjectName, const FTimecode& Timecode, uint32 FrameRate, const FARBlendShapeMap& BlendShapes, FName DeviceID = NAME_None) = 0;
+	virtual void PublishBlendShapes(FName SubjectName, const FQualifiedFrameTime& FrameTime, const FARBlendShapeMap& BlendShapes, FName DeviceID = NAME_None) = 0;
 };
 
 /** Interface that publishes face ar blend shape information via LiveLink */
