LWC: AppleARKitFaceMeshConversion ToVertexBuffer return type fix

diff --git a/Engine/Plugins/Runtime/AR/AppleAR/AppleARKitFaceSupport/Source/AppleARKitFaceSupport/Private/AppleARKitFaceMeshConversion.h b/Engine/Plugins/Runtime/AR/AppleAR/AppleARKitFaceSupport/Source/AppleARKitFaceSupport/Private/AppleARKitFaceMeshConversion.h
--- a/Engine/Plugins/Runtime/AR/AppleAR/AppleARKitFaceSupport/Source/AppleARKitFaceSupport/Private/AppleARKitFaceMeshConversion.h
+++ b/Engine/Plugins/Runtime/AR/AppleAR/AppleARKitFaceSupport/Source/AppleARKitFaceSupport/Private/AppleARKitFaceMeshConversion.h
@@ -23,14 +23,14 @@ static FORCEINLINE TArray<int32> To32BitIndexBuffer(const int16_t* Indices, cons
 }
 
 // @todo JoeG - An option for which way to orient tris (down +X or -X)
-static FORCEINLINE TArray<FVector3f> ToVertexBuffer(const vector_float3* Vertices, const uint64 VertexCount)
+static FORCEINLINE TArray<FVector> ToVertexBuffer(const vector_float3* Vertices, const uint64 VertexCount)
 {
-	TArray<FVector3f> VertexBuffer;
+	TArray<FVector> VertexBuffer;
 	VertexBuffer.AddUninitialized(VertexCount);
 	// @todo JoeG - make a fast routine for this
 	for (int32 Index = 0; Index < VertexBuffer.Num(); Index++)
 	{
-		VertexBuffer[Index] = FVector3f(Vertices[Index].z, Vertices[Index].x, Vertices[Index].y);
+		VertexBuffer[Index] = FVector(Vertices[Index].z, Vertices[Index].x, Vertices[Index].y);
 	}
 	return VertexBuffer;
 }
