UE4 -> UE references in variable names and comments for ARKit.

diff --git a/Engine/Plugins/Runtime/AR/AppleAR/AppleARKit/Source/AppleARKit/Private/AppleARKitCamera.cpp b/Engine/Plugins/Runtime/AR/AppleAR/AppleARKit/Source/AppleARKit/Private/AppleARKitCamera.cpp
--- a/Engine/Plugins/Runtime/AR/AppleAR/AppleARKit/Source/AppleARKit/Private/AppleARKitCamera.cpp
+++ b/Engine/Plugins/Runtime/AR/AppleAR/AppleARKit/Source/AppleARKit/Private/AppleARKitCamera.cpp
@@ -5,7 +5,7 @@
 #include "AppleARKitModule.h"
 #include "AppleARKitConversion.h"
 	
-// UE4
+// UE
 #include "CoreGlobals.h"
 #include "Engine/EngineTypes.h"
 #include "Kismet/GameplayStatics.h"
diff --git a/Engine/Plugins/Runtime/AR/AppleAR/AppleARKit/Source/AppleARKit/Private/AppleARKitCamera.h b/Engine/Plugins/Runtime/AR/AppleAR/AppleARKit/Source/AppleARKit/Private/AppleARKitCamera.h
--- a/Engine/Plugins/Runtime/AR/AppleAR/AppleARKit/Source/AppleARKit/Private/AppleARKitCamera.h
+++ b/Engine/Plugins/Runtime/AR/AppleAR/AppleARKit/Source/AppleARKit/Private/AppleARKitCamera.h
@@ -2,7 +2,7 @@
 
 #pragma once
 
-// UE4
+// UE
 #include "Math/Transform.h"
 #include "Math/UnrealMathUtility.h"
 #include "UnrealEngine.h"
@@ -54,7 +54,7 @@ struct APPLEARKIT_API FAppleARKitCamera
 
 	/** 
 	 * This is a conversion copy-constructor that takes a raw ARCamera and fills this structs members
-	 * with the UE4-ified versions of ARCamera's properties.
+	 * with the UE-ified versions of ARCamera's properties.
 	 */ 
 	FAppleARKitCamera( ARCamera* InARCamera );
 
diff --git a/Engine/Plugins/Runtime/AR/AppleAR/AppleARKit/Source/AppleARKit/Private/AppleARKitFrame.cpp b/Engine/Plugins/Runtime/AR/AppleAR/AppleARKit/Source/AppleARKit/Private/AppleARKitFrame.cpp
--- a/Engine/Plugins/Runtime/AR/AppleAR/AppleARKit/Source/AppleARKit/Private/AppleARKitFrame.cpp
+++ b/Engine/Plugins/Runtime/AR/AppleAR/AppleARKit/Source/AppleARKit/Private/AppleARKitFrame.cpp
@@ -89,7 +89,7 @@ FAppleARKitFrame::FAppleARKitFrame(ARFrame* InARFrame, const FVector2D MinCamera
 		{
 			Tracked2DPose = FAppleARKitConversion::ToARPose2D(InARFrame.detectedBody);
 			
-			// Convert the joint location from the normalized arkit camera space to UE4's normalized screen space
+			// Convert the joint location from the normalized arkit camera space to UE's normalized screen space
 			const FVector2D UVSize = MaxCameraUV - MinCameraUV;
 			for (int Index = 0; Index < Tracked2DPose.JointLocations.Num(); ++Index)
 			{
diff --git a/Engine/Plugins/Runtime/AR/AppleAR/AppleARKitFaceSupport/Source/AppleARKitFaceSupport/Private/AppleARKitFaceMeshComponent.cpp b/Engine/Plugins/Runtime/AR/AppleAR/AppleARKitFaceSupport/Source/AppleARKitFaceSupport/Private/AppleARKitFaceMeshComponent.cpp
--- a/Engine/Plugins/Runtime/AR/AppleAR/AppleARKitFaceSupport/Source/AppleARKitFaceSupport/Private/AppleARKitFaceMeshComponent.cpp
+++ b/Engine/Plugins/Runtime/AR/AppleAR/AppleARKitFaceSupport/Source/AppleARKitFaceSupport/Private/AppleARKitFaceMeshComponent.cpp
@@ -23,10 +23,10 @@ NSDictionary<ARBlendShapeLocation,NSNumber *>* ToBlendShapeDictionary(const FARB
 {
 	NSMutableDictionary<ARBlendShapeLocation,NSNumber *>* BlendShapeDict = [[[NSMutableDictionary<ARBlendShapeLocation,NSNumber *> alloc] init] autorelease];
 
-#define SET_BLEND_SHAPE(AppleShape, UE4Shape) \
-	if (BlendShapeMap.Contains(UE4Shape)) \
+#define SET_BLEND_SHAPE(AppleShape, UEShape) \
+	if (BlendShapeMap.Contains(UEShape)) \
 	{ \
-		NSNumber* Num = [NSNumber numberWithFloat: BlendShapeMap[UE4Shape]]; \
+		NSNumber* Num = [NSNumber numberWithFloat: BlendShapeMap[UEShape]]; \
 		BlendShapeDict[AppleShape] = Num; \
 	}
 
diff --git a/Engine/Plugins/Runtime/AR/AppleAR/AppleARKitFaceSupport/Source/AppleARKitFaceSupport/Private/AppleARKitFaceMeshConversion.h b/Engine/Plugins/Runtime/AR/AppleAR/AppleARKitFaceSupport/Source/AppleARKitFaceSupport/Private/AppleARKitFaceMeshConversion.h
--- a/Engine/Plugins/Runtime/AR/AppleAR/AppleARKitFaceSupport/Source/AppleARKitFaceSupport/Private/AppleARKitFaceMeshConversion.h
+++ b/Engine/Plugins/Runtime/AR/AppleAR/AppleARKitFaceSupport/Source/AppleARKitFaceSupport/Private/AppleARKitFaceMeshConversion.h
@@ -77,14 +77,14 @@ static FORCEINLINE FARBlendShapeMap ToBlendShapeMap(bool bFaceMirrored, NSDictio
 	BlendShapeMap.Add(EARFaceBlendShape::RightEyeRoll, RightEyeRoll);
 
 
-#define SET_BLEND_SHAPE(AppleShape, UE4Shape) \
+#define SET_BLEND_SHAPE(AppleShape, UEShape) \
 	if (BlendShapes[AppleShape]) \
 	{ \
-		BlendShapeMap.Add(UE4Shape, FMath::Max([BlendShapes[AppleShape] floatValue], 0.f)); \
+		BlendShapeMap.Add(UEShape, FMath::Max([BlendShapes[AppleShape] floatValue], 0.f)); \
 	} \
 	else \
 	{ \
-		BlendShapeMap.Add(UE4Shape, 0.f); \
+		BlendShapeMap.Add(UEShape, 0.f); \
 	}
 
 	if (bFaceMirrored)
