Fixed mismatched copy constructor/assignment operator in FNetQuantizeFaceCurve.

diff --git a/Engine/Plugins/Runtime/AR/AppleAR/AppleARKitFaceSupport/Source/AppleARKitFaceSupport/Public/AppleARKitFaceMeshComponent.h b/Engine/Plugins/Runtime/AR/AppleAR/AppleARKitFaceSupport/Source/AppleARKitFaceSupport/Public/AppleARKitFaceMeshComponent.h
--- a/Engine/Plugins/Runtime/AR/AppleAR/AppleARKitFaceSupport/Source/AppleARKitFaceSupport/Public/AppleARKitFaceMeshComponent.h
+++ b/Engine/Plugins/Runtime/AR/AppleAR/AppleARKitFaceSupport/Source/AppleARKitFaceSupport/Public/AppleARKitFaceMeshComponent.h
@@ -30,28 +30,22 @@ USTRUCT()
 struct FNetQuantizeFaceCurve
 {
 	GENERATED_USTRUCT_BODY()
-	
+
 	FORCEINLINE FNetQuantizeFaceCurve()
 	{
 	}
-	
+
 	FORCEINLINE FNetQuantizeFaceCurve(EARFaceBlendShape InBlendShape, float InAmount)
 		: BlendShape((uint8)InBlendShape)
 		, Amount(ConvertAmountToInt(InAmount))
 	{
 	}
-	
-	FORCEINLINE FNetQuantizeFaceCurve(const FNetQuantizeFaceCurve& Other)
-		: BlendShape(Other.BlendShape)
-		, Amount(Other.Amount)
-	{
-	}
-	
+
 	FORCEINLINE float GetAmountAsFloat() const
 	{
 		return ConvertAmountToFloat(Amount);
 	}
-	
+
 	FORCEINLINE EARFaceBlendShape GetBlendShape() const
 	{
 		check(BlendShape < (uint8)EARFaceBlendShape::MAX);
@@ -75,7 +69,7 @@ struct FNetQuantizeFaceCurve
 		bOutSuccess = true;
 		return true;
 	}
-	
+
 	static bool IsDifferentEnough(float Val1, float Val2)
 	{
 		return ConvertAmountToInt(Val1) != ConvertAmountToInt(Val2);
@@ -89,12 +83,12 @@ private:
 	{
 		return ((float)InAmount) * InvScale;
 	}
-	
+
 	static FORCEINLINE int8 ConvertAmountToInt(float InAmount)
 	{
 		return (int8)FMath::TruncToInt(InAmount * Scale);
 	}
-	
+
 	uint8 BlendShape;
 	int8 Amount;
 };
