Enable support for andriod platform.

diff --git a/Engine/Plugins/Runtime/AR/AppleAR/AppleARKit/AppleARKit.uplugin b/Engine/Plugins/Runtime/AR/AppleAR/AppleARKit/AppleARKit.uplugin
--- a/Engine/Plugins/Runtime/AR/AppleAR/AppleARKit/AppleARKit.uplugin
+++ b/Engine/Plugins/Runtime/AR/AppleAR/AppleARKit/AppleARKit.uplugin
@@ -14,20 +14,20 @@
 	"IsBetaVersion": false,
 	"EnabledByDefault": false,
 	"SupportedTargetPlatforms" : [
-		"IOS", "Win64", "Mac","Linux"
+		"IOS", "Win64", "Mac","Linux", "Android"
 	],
 	"Modules": [
 		{
 			"Name": "AppleARKit",
 			"Type": "Runtime",
 			"LoadingPhase": "PostConfigInit",
-			"PlatformAllowList" : [ "Mac", "IOS", "Win64","Linux" ]
+			"PlatformAllowList" : [ "Mac", "IOS", "Win64","Linux", "Android" ]
 		},
 		{
 			"Name": "AppleARKitPoseTrackingLiveLink",
 			"Type": "Runtime",
 			"LoadingPhase": "PostConfigInit",
-			"PlatformAllowList" : [ "Mac", "IOS", "Win64" ,"Linux"]
+			"PlatformAllowList" : [ "Mac", "IOS", "Win64" ,"Linux", "Android" ]
 
 		}
 	],
diff --git a/Engine/Plugins/Runtime/AR/AppleAR/AppleARKitFaceSupport/AppleARKitFaceSupport.uplugin b/Engine/Plugins/Runtime/AR/AppleAR/AppleARKitFaceSupport/AppleARKitFaceSupport.uplugin
--- a/Engine/Plugins/Runtime/AR/AppleAR/AppleARKitFaceSupport/AppleARKitFaceSupport.uplugin
+++ b/Engine/Plugins/Runtime/AR/AppleAR/AppleARKitFaceSupport/AppleARKitFaceSupport.uplugin
@@ -14,7 +14,7 @@
 	"IsBetaVersion": false,
 	"EnabledByDefault": false,
 	"SupportedTargetPlatforms" : [
-		"IOS", "Win64", "Mac","Linux"
+		"IOS", "Win64", "Mac","Linux", "Android"
 	],
 	"Modules":
 	[
@@ -22,7 +22,7 @@
 			"Name": "AppleARKitFaceSupport",
 			"Type": "Runtime",
 			"LoadingPhase": "PostConfigInit",
-			"PlatformAllowList" : [ "Mac", "IOS", "Win64" ,"Linux"]
+			"PlatformAllowList" : [ "Mac", "IOS", "Win64" ,"Linux", "Android"]
 		}
 	],
     "Plugins":
diff --git a/Engine/Plugins/Runtime/AR/AppleAR/AppleARKitFaceSupport/Source/AppleARKitFaceSupport/Private/AppleARKitFaceSupportModule.cpp b/Engine/Plugins/Runtime/AR/AppleAR/AppleARKitFaceSupport/Source/AppleARKitFaceSupport/Private/AppleARKitFaceSupportModule.cpp
--- a/Engine/Plugins/Runtime/AR/AppleAR/AppleARKitFaceSupport/Source/AppleARKitFaceSupport/Private/AppleARKitFaceSupportModule.cpp
+++ b/Engine/Plugins/Runtime/AR/AppleAR/AppleARKitFaceSupport/Source/AppleARKitFaceSupport/Private/AppleARKitFaceSupportModule.cpp
@@ -22,8 +22,8 @@ void FAppleARKitFaceSupportModule::StartupModule()
 	FaceSupportInstance = MakeShared<FAppleARKitFaceSupport, ESPMode::ThreadSafe>();
 	FaceSupportInstance->Init();
 
-	// LiveLink listener needs to be created here so that the editor can receive remote publishing events
-#if PLATFORM_DESKTOP
+	// LiveLink listener needs to be created here so that we can receive remote publishing events
+#if PLATFORM_DESKTOP || PLATFORM_ANDROID
 	bool bEnableLiveLinkForFaceTracking = false;
 	GConfig->GetBool(TEXT("/Script/AppleARKit.AppleARKitSettings"), TEXT("bEnableLiveLinkForFaceTracking"), bEnableLiveLinkForFaceTracking, GEngineIni);
 	if (bEnableLiveLinkForFaceTracking)
