UE-193715 AppleARKitFaceSupport & AppleARKitPoseTrackingLiveLink reference the Private folder of the AppleARKit module
-Removed private folder access along with the includes, which are unneeded.

diff --git a/Engine/Plugins/Runtime/AR/AppleAR/AppleARKitFaceSupport/Source/AppleARKitFaceSupport/AppleARKitFaceSupport.Build.cs b/Engine/Plugins/Runtime/AR/AppleAR/AppleARKitFaceSupport/Source/AppleARKitFaceSupport/AppleARKitFaceSupport.Build.cs
--- a/Engine/Plugins/Runtime/AR/AppleAR/AppleARKitFaceSupport/Source/AppleARKitFaceSupport/AppleARKitFaceSupport.Build.cs
+++ b/Engine/Plugins/Runtime/AR/AppleAR/AppleARKitFaceSupport/Source/AppleARKitFaceSupport/AppleARKitFaceSupport.Build.cs
@@ -6,14 +6,6 @@ public class AppleARKitFaceSupport : ModuleRules
 {
 	public AppleARKitFaceSupport(ReadOnlyTargetRules Target) : base(Target)
 	{
-		PrivateIncludePaths.AddRange(
-			new string[] {
-				System.IO.Path.Combine(GetModuleDirectory("AppleARKit"), "Private"),
-				// ... add other private include paths required here ...
-			}
-			);
-			
-		
 		PublicDependencyModuleNames.AddRange(
 			new string[]
 			{
diff --git a/Engine/Plugins/Runtime/AR/AppleAR/AppleARKitFaceSupport/Source/AppleARKitFaceSupport/Private/AppleARKitFaceSupportModule.cpp b/Engine/Plugins/Runtime/AR/AppleAR/AppleARKitFaceSupport/Source/AppleARKitFaceSupport/Private/AppleARKitFaceSupportModule.cpp
--- a/Engine/Plugins/Runtime/AR/AppleAR/AppleARKitFaceSupport/Source/AppleARKitFaceSupport/Private/AppleARKitFaceSupportModule.cpp
+++ b/Engine/Plugins/Runtime/AR/AppleAR/AppleARKitFaceSupport/Source/AppleARKitFaceSupport/Private/AppleARKitFaceSupportModule.cpp
@@ -1,7 +1,6 @@
 // Copyright Epic Games, Inc. All Rights Reserved.
 
 #include "AppleARKitFaceSupportModule.h"
-#include "AppleARKitSystem.h"
 #include "Features/IModularFeature.h"
 #include "Features/IModularFeatures.h"
 #include "CoreGlobals.h"
diff --git a/Engine/Plugins/Runtime/AR/AppleAR/AppleARKitFaceSupport/Source/AppleARKitFaceSupport/Private/AppleARKitLiveLinkSource.cpp b/Engine/Plugins/Runtime/AR/AppleAR/AppleARKitFaceSupport/Source/AppleARKitFaceSupport/Private/AppleARKitLiveLinkSource.cpp
--- a/Engine/Plugins/Runtime/AR/AppleAR/AppleARKitFaceSupport/Source/AppleARKitFaceSupport/Private/AppleARKitLiveLinkSource.cpp
+++ b/Engine/Plugins/Runtime/AR/AppleAR/AppleARKitFaceSupport/Source/AppleARKitFaceSupport/Private/AppleARKitLiveLinkSource.cpp
@@ -8,7 +8,6 @@
 #include "UObject/ObjectMacros.h"
 #include "ARSystem.h"
 #include "ARBlueprintLibrary.h"
-#include "AppleARKitModule.h"
 #include "Features/IModularFeatures.h"
 #include "Misc/FileHelper.h"
 
