AppleARKit - By default, prevent ARKit from starting its listener.

- You can override this to revert to the previous behavior by specifying -EnableLiveLinkFaceTracking on the command line.
- Modify ARKit module to not bind to its port by default unless specified by the command line
- Move ARKit server binding to its source instead of the module

This fixes an issue where starting two editors would make it impossible for one of them to use the ARKit livelink plugin since only one could bind to the ARKit port.
Moving it to the livelink source allows the user to pick which UE instance should bind the port.

diff --git a/Engine/Config/BaseEngine.ini b/Engine/Config/BaseEngine.ini
--- a/Engine/Config/BaseEngine.ini
+++ b/Engine/Config/BaseEngine.ini
@@ -3398,7 +3398,7 @@ PakSigningRequired=True
 PakEncryptionRequired=True
 
 [/Script/AppleARKit.AppleARKitSettings]
-bEnableLiveLinkForFaceTracking=true
+bEnableLiveLinkForFaceTracking=false
 LiveLinkPublishingPort=11111
 ; by default arkit will be added to UIRequiredDeviceCapabilities and devices must support
 ; arkit. If false games must check that arkit is supported before enabling arkit features
diff --git a/Engine/Plugins/Runtime/AR/AppleAR/AppleARKitFaceSupport/Source/AppleARKitFaceSupport/AppleARKitFaceSupport.Build.cs b/Engine/Plugins/Runtime/AR/AppleAR/AppleARKitFaceSupport/Source/AppleARKitFaceSupport/AppleARKitFaceSupport.Build.cs
--- a/Engine/Plugins/Runtime/AR/AppleAR/AppleARKitFaceSupport/Source/AppleARKitFaceSupport/AppleARKitFaceSupport.Build.cs
+++ b/Engine/Plugins/Runtime/AR/AppleAR/AppleARKitFaceSupport/Source/AppleARKitFaceSupport/AppleARKitFaceSupport.Build.cs
@@ -45,5 +45,17 @@ public class AppleARKitFaceSupport : ModuleRules
 		{
 			PublicFrameworks.Add( "ARKit" );
 		}
+		
+		if (Target.bCompileAgainstEditor)
+		{
+			PrivateDependencyModuleNames.AddRange(
+				new string[]
+				{
+					"Slate",
+					"SlateCore",
+					"PropertyEditor"
+				}
+			);
+		}
 	}
 }
diff --git a/Engine/Plugins/Runtime/AR/AppleAR/AppleARKitFaceSupport/Source/AppleARKitFaceSupport/Private/AppleARKitFaceSupportImpl.cpp b/Engine/Plugins/Runtime/AR/AppleAR/AppleARKitFaceSupport/Source/AppleARKitFaceSupport/Private/AppleARKitFaceSupportImpl.cpp
--- a/Engine/Plugins/Runtime/AR/AppleAR/AppleARKitFaceSupport/Source/AppleARKitFaceSupport/Private/AppleARKitFaceSupportImpl.cpp
+++ b/Engine/Plugins/Runtime/AR/AppleAR/AppleARKitFaceSupport/Source/AppleARKitFaceSupport/Private/AppleARKitFaceSupportImpl.cpp
@@ -292,17 +292,19 @@ void FAppleARKitFaceSupport::PublishLiveLinkData(const FGuid& SessionGuid, TShar
 #if PLATFORM_IOS
 			LiveLinkSource = FAppleARKitLiveLinkSourceFactory::CreateLiveLinkSource();
 #else
-			// This should be started already, but just in case
-			FAppleARKitLiveLinkSourceFactory::CreateLiveLinkRemoteListener();
+			// This will initialize the remote listener.
+			LiveLinkSource = FAppleARKitLiveLinkSourceFactory::CreateLiveLinkSource(FAppleARKitLiveLinkConnectionSettings());
 #endif
 		}
 	}
 	
+#if PLATFORM_IOS
 	if (LiveLinkSource.IsValid())
 	{
 		const FQualifiedFrameTime FrameTime(Anchor->Timecode, FFrameRate(Anchor->FrameRate, 1));
 		LiveLinkSource->PublishBlendShapes(GetLiveLinkSubjectName(Anchor->AnchorGUID), FrameTime, Anchor->BlendShapes, LocalDeviceId);
 	}
+#endif
 }
 
 bool FAppleARKitFaceSupport::DoesSupportFaceAR()
diff --git a/Engine/Plugins/Runtime/AR/AppleAR/AppleARKitFaceSupport/Source/AppleARKitFaceSupport/Private/AppleARKitFaceSupportModule.cpp b/Engine/Plugins/Runtime/AR/AppleAR/AppleARKitFaceSupport/Source/AppleARKitFaceSupport/Private/AppleARKitFaceSupportModule.cpp
--- a/Engine/Plugins/Runtime/AR/AppleAR/AppleARKitFaceSupport/Source/AppleARKitFaceSupport/Private/AppleARKitFaceSupportModule.cpp
+++ b/Engine/Plugins/Runtime/AR/AppleAR/AppleARKitFaceSupport/Source/AppleARKitFaceSupport/Private/AppleARKitFaceSupportModule.cpp
@@ -1,19 +1,19 @@
 // Copyright Epic Games, Inc. All Rights Reserved.
 
 #include "AppleARKitFaceSupportModule.h"
-#include "Features/IModularFeature.h"
-#include "Features/IModularFeatures.h"
-#include "CoreGlobals.h"
+
 #include "AppleARKitFaceSupportImpl.h"
+#include "AppleARKitLiveLinkSource.h"
+#include "CoreGlobals.h"
+#include "Features/IModularFeature.h"
 #include "Misc/ConfigCacheIni.h"
+#include "Misc/CoreDelegates.h"
 
 
 IMPLEMENT_MODULE(FAppleARKitFaceSupportModule, AppleARKitFaceSupport);
 
 DEFINE_LOG_CATEGORY(LogAppleARKitFace);
 
-TSharedPtr<FAppleARKitFaceSupport, ESPMode::ThreadSafe> FaceSupportInstance;
-
 void FAppleARKitFaceSupportModule::StartupModule()
 {
 	ensureMsgf(FModuleManager::Get().LoadModule("AppleARKit"), TEXT("ARKitFaceSupport depends on the AppleARKit module."));
@@ -21,20 +21,37 @@ void FAppleARKitFaceSupportModule::StartupModule()
 	FaceSupportInstance = MakeShared<FAppleARKitFaceSupport, ESPMode::ThreadSafe>();
 	FaceSupportInstance->Init();
 
-	// LiveLink listener needs to be created here so that we can receive remote publishing events
 #if PLATFORM_DESKTOP || PLATFORM_ANDROID
-	bool bEnableLiveLinkForFaceTracking = false;
-	GConfig->GetBool(TEXT("/Script/AppleARKit.AppleARKitSettings"), TEXT("bEnableLiveLinkForFaceTracking"), bEnableLiveLinkForFaceTracking, GEngineIni);
-	if (bEnableLiveLinkForFaceTracking)
-	{
-		FAppleARKitLiveLinkSourceFactory::CreateLiveLinkRemoteListener();
-	}
+	IModularFeatures::Get().OnModularFeatureRegistered().AddRaw(this, &FAppleARKitFaceSupportModule::OnModularFeatureAvailable);
 #endif
 }
 
 void FAppleARKitFaceSupportModule::ShutdownModule()
 {
+	LiveLinkSource.Reset();
 	FaceSupportInstance->Shutdown();
 	FaceSupportInstance = nullptr;
 }
 
+void FAppleARKitFaceSupportModule::OnModularFeatureAvailable(const FName& Type, class IModularFeature* ModularFeature)
+{
+	// LiveLink listener needs to be created here so that we can receive remote publishing events
+	if (Type == ILiveLinkClient::ModularFeatureName)
+	{
+		// Reset in case the modular feature was re-registered (ie. LiveLinkHub).
+		LiveLinkSource.Reset();
+
+		bool bEnableLiveLinkForFaceTracking = false;
+		FAppleARKitLiveLinkConnectionSettings Settings;
+		GConfig->GetBool(TEXT("/Script/AppleARKit.AppleARKitSettings"), TEXT("bEnableLiveLinkForFaceTracking"), bEnableLiveLinkForFaceTracking, GEngineIni);
+		GConfig->GetInt(TEXT("/Script/AppleARKit.AppleARKitSettings"), TEXT("LiveLinkPublishingPort"), Settings.Port, GEngineIni);
+
+		if (bEnableLiveLinkForFaceTracking || FParse::Param(FCommandLine::Get(), TEXT("EnableLiveLinkFaceTracking")))
+		{
+			LiveLinkSource = FAppleARKitLiveLinkSourceFactory::CreateLiveLinkSource(Settings);
+		}
+	}
+}
+
+
+
diff --git a/Engine/Plugins/Runtime/AR/AppleAR/AppleARKitFaceSupport/Source/AppleARKitFaceSupport/Private/AppleARKitFaceSupportModule.h b/Engine/Plugins/Runtime/AR/AppleAR/AppleARKitFaceSupport/Source/AppleARKitFaceSupport/Private/AppleARKitFaceSupportModule.h
--- a/Engine/Plugins/Runtime/AR/AppleAR/AppleARKitFaceSupport/Source/AppleARKitFaceSupport/Private/AppleARKitFaceSupportModule.h
+++ b/Engine/Plugins/Runtime/AR/AppleAR/AppleARKitFaceSupport/Source/AppleARKitFaceSupport/Private/AppleARKitFaceSupportModule.h
@@ -9,9 +9,20 @@ class APPLEARKITFACESUPPORT_API FAppleARKitFaceSupportModule :
 	public IModuleInterface
 {
 public:
+	//~ Begin IModuleInterface
 	virtual void StartupModule() override;
-
 	virtual void ShutdownModule() override;
+	//~ End IModuleInterface
+
+private:
+	/** Handler called when a modular feature is made available. Used to start a remote listener after the livelink client is registered. */
+	void OnModularFeatureAvailable(const FName& Type, class IModularFeature* ModularFeature);
+private:
+	// Interface / helper for sending AR Face data to LiveLink from an IOS device.
+	TSharedPtr<class FAppleARKitFaceSupport, ESPMode::ThreadSafe> FaceSupportInstance;
+
+	// LiveLinkSource responsible for receiving remote LiveLink data and publishing it to the LiveLink client.
+	TSharedPtr<class ILiveLinkSourceARKit> LiveLinkSource;
 };
 
 DECLARE_LOG_CATEGORY_EXTERN(LogAppleARKitFace, Log, All);
diff --git a/Engine/Plugins/Runtime/AR/AppleAR/AppleARKitFaceSupport/Source/AppleARKitFaceSupport/Private/AppleARKitLiveLinkConnectionSettings.h b/Engine/Plugins/Runtime/AR/AppleAR/AppleARKitFaceSupport/Source/AppleARKitFaceSupport/Private/AppleARKitLiveLinkConnectionSettings.h
new file mode 100644
--- /dev/null
+++ b/Engine/Plugins/Runtime/AR/AppleAR/AppleARKitFaceSupport/Source/AppleARKitFaceSupport/Private/AppleARKitLiveLinkConnectionSettings.h
@@ -0,0 +1,15 @@
+// Copyright Epic Games, Inc. All Rights Reserved.
+
+#pragma once
+
+#include "AppleARKitLiveLinkConnectionSettings.generated.h"
+
+USTRUCT()
+struct APPLEARKITFACESUPPORT_API FAppleARKitLiveLinkConnectionSettings
+{
+	GENERATED_BODY()
+
+	/** The port to use when listening/sending LiveLink face blend shapes via the network */
+	UPROPERTY(EditAnywhere, Category = "Connection Settings")
+	int32 Port = 11111;
+};
diff --git a/Engine/Plugins/Runtime/AR/AppleAR/AppleARKitFaceSupport/Source/AppleARKitFaceSupport/Private/AppleARKitLiveLinkSource.cpp b/Engine/Plugins/Runtime/AR/AppleAR/AppleARKitFaceSupport/Source/AppleARKitFaceSupport/Private/AppleARKitLiveLinkSource.cpp
--- a/Engine/Plugins/Runtime/AR/AppleAR/AppleARKitFaceSupport/Source/AppleARKitFaceSupport/Private/AppleARKitLiveLinkSource.cpp
+++ b/Engine/Plugins/Runtime/AR/AppleAR/AppleARKitFaceSupport/Source/AppleARKitFaceSupport/Private/AppleARKitLiveLinkSource.cpp
@@ -16,6 +16,7 @@
 #include "Sockets.h"
 
 #include "AppleARKitSettings.h"
+#include "AppleARKitLiveLinkConnectionSettings.h"
 #include "ARTrackable.h"
 
 #include "ILiveLinkClient.h"
@@ -54,18 +55,21 @@ TSharedPtr<ILiveLinkSourceARKit> FAppleARKitLiveLinkSourceFactory::CreateLiveLin
 	return TSharedPtr<ILiveLinkSourceARKit>();
 }
 
-void FAppleARKitLiveLinkSourceFactory::CreateLiveLinkRemoteListener()
+TSharedPtr<ILiveLinkSourceARKit> FAppleARKitLiveLinkSourceFactory::CreateLiveLinkSource(const FAppleARKitLiveLinkConnectionSettings& ConnectionSettings)
 {
-	static FAppleARKitLiveLinkRemoteListener* Listener = nullptr;
-	if (Listener == nullptr)
+	IModularFeatures& ModularFeatures = IModularFeatures::Get();
+
+	if (ModularFeatures.IsModularFeatureAvailable(ILiveLinkClient::ModularFeatureName))
 	{
-		Listener = new FAppleARKitLiveLinkRemoteListener();
-		if (!Listener->InitReceiveSocket())
-		{
-			delete Listener;
-			Listener = nullptr;
-		}
+		ILiveLinkClient* LiveLinkClient = &IModularFeatures::Get().GetModularFeature<ILiveLinkClient>(ILiveLinkClient::ModularFeatureName);
+		TSharedPtr<FAppleARKitLiveLinkSource> Source = MakeShared<FAppleARKitLiveLinkSource>(ConnectionSettings);
+		Source->InitializeRemoteListener();
+		
+		LiveLinkClient->AddSource(Source);
+
+		return Source;
 	}
+	return TSharedPtr<ILiveLinkSourceARKit>();
 }
 
 TSharedPtr<IARKitBlendShapePublisher, ESPMode::ThreadSafe> FAppleARKitLiveLinkSourceFactory::CreateLiveLinkRemotePublisher(const FString& RemoteAddr)
@@ -109,18 +113,44 @@ TSharedPtr<IARKitBlendShapePublisher, ESPMode::ThreadSafe> FAppleARKitLiveLinkSo
 	return LocalFileWriter;
 }
 
-FAppleARKitLiveLinkSource::FAppleARKitLiveLinkSource() :
-	Client(nullptr)
+FAppleARKitLiveLinkSource::FAppleARKitLiveLinkSource() 
+	: Client(nullptr)
 	, LastFramePublished(0)
 {
 }
 
+FAppleARKitLiveLinkSource::FAppleARKitLiveLinkSource(FAppleARKitLiveLinkConnectionSettings InConnectionSettings)
+	: Client(nullptr)
+	, LastFramePublished(0)
+	, ConnectionSettings(MoveTemp(InConnectionSettings))
+{
+}
+
+void FAppleARKitLiveLinkSource::InitializeRemoteListener()
+{
+	// LiveLink listener needs to be created here so that we can receive remote publishing events
+#if PLATFORM_DESKTOP || PLATFORM_ANDROID
+	RemoteListener = MakeUnique<FAppleARKitLiveLinkRemoteListener>(AsShared());
+	if (!RemoteListener->InitReceiveSocket(ConnectionSettings.Port))
+	{
+		RemoteListener.Reset();
+		UE_LOG(LogAppleARKitFace, Warning, TEXT("Could not start AppleARKit remote listener on port %d"), ConnectionSettings.Port);
+	}
+#endif
+}
+
 void FAppleARKitLiveLinkSource::ReceiveClient(ILiveLinkClient* InClient, FGuid InSourceGuid)
 {
 	Client = InClient;
 	SourceGuid = InSourceGuid;
 }
 
+void FAppleARKitLiveLinkSource::InitializeSettings(ULiveLinkSourceSettings* Settings)
+{
+	FAppleARKitLiveLinkConnectionSettings::StaticStruct()->ExportText(Settings->ConnectionString, &ConnectionSettings, nullptr, nullptr, PPF_None, nullptr);
+	Settings->Factory = UAppleARKitLiveLinkSourceFactory::StaticClass();
+}
+
 bool FAppleARKitLiveLinkSource::IsSourceStillValid() const
 {
 	return Client != nullptr;
@@ -129,6 +159,7 @@ bool FAppleARKitLiveLinkSource::IsSourceStillValid() const
 bool FAppleARKitLiveLinkSource::RequestSourceShutdown()
 {
 	Client = nullptr;
+	RemoteListener.Reset();
 	return true;
 }
 
@@ -144,7 +175,10 @@ FText FAppleARKitLiveLinkSource::GetSourceStatus() const
 
 FText FAppleARKitLiveLinkSource::GetSourceType() const
 {
-	return NSLOCTEXT( "AppleARKitLiveLink", "AppleARKitLiveLinkSourceType", "Apple AR Face Tracking" );
+	FNumberFormattingOptions Options;
+	Options.UseGrouping = false;
+
+	return FText::Format(NSLOCTEXT("AppleARKitLiveLink", "AppleARKitLiveLinkSourceType", "ARKit Port: {0}"), FText::AsNumber(ConnectionSettings.Port, &Options));
 }
 
 static FName ParseEnumName(FName EnumName)
@@ -351,8 +385,9 @@ void FAppleARKitLiveLinkRemotePublisher::PublishBlendShapes(FName SubjectName, c
 	}
 }
 
-FAppleARKitLiveLinkRemoteListener::FAppleARKitLiveLinkRemoteListener() :
-	RecvSocket(nullptr)
+FAppleARKitLiveLinkRemoteListener::FAppleARKitLiveLinkRemoteListener(const TSharedPtr<ILiveLinkSourceARKit>& ArkitSource)
+	: RecvSocket(nullptr)
+	, Source(ArkitSource)
 {
 	RecvBuffer.AddUninitialized(MAX_BLEND_SHAPE_PACKET_SIZE);
 }
@@ -371,14 +406,12 @@ FAppleARKitLiveLinkRemoteListener::~FAppleARKitLiveLinkRemoteListener()
 	}
 }
 
-bool FAppleARKitLiveLinkRemoteListener::InitReceiveSocket()
+bool FAppleARKitLiveLinkRemoteListener::InitReceiveSocket(int32 Port)
 {
 	ISocketSubsystem* SocketSubsystem = ISocketSubsystem::Get();
 	TSharedRef<FInternetAddr> Addr = SocketSubsystem->GetLocalBindAddr(*GLog);
-	int32 LiveLinkPort = 0;
-	// Have to read this value manually since it happens before UObjects are alive
-	GConfig->GetInt(TEXT("/Script/AppleARKit.AppleARKitSettings"), TEXT("LiveLinkPublishingPort"), LiveLinkPort, GEngineIni);
-	Addr->SetPort(LiveLinkPort);
+
+	Addr->SetPort(Port);
 
 	RecvSocket = SocketSubsystem->CreateSocket(NAME_DGram, TEXT("FAppleARKitLiveLinkRemoteListener socket"), Addr->GetProtocolType());
 	if (RecvSocket != nullptr)
@@ -398,14 +431,6 @@ bool FAppleARKitLiveLinkRemoteListener::InitReceiveSocket()
 	return RecvSocket != nullptr;
 }
 
-void FAppleARKitLiveLinkRemoteListener::InitLiveLinkSource()
-{
-	if (!Source.IsValid() || !Source->IsSourceStillValid())
-	{
-		Source = FAppleARKitLiveLinkSourceFactory::CreateLiveLinkSource();
-	}
-}
-
 void FAppleARKitLiveLinkRemoteListener::Tick(float DeltaTime)
 {
 	SCOPE_CYCLE_COUNTER(STAT_FaceAR_ReceiveLiveLink);
@@ -460,10 +485,12 @@ void FAppleARKitLiveLinkRemoteListener::Tick(float DeltaTime)
 			// All of the data was valid, so publish it
 			if (!FromBuffer.HasOverflow())
 			{
-				InitLiveLinkSource();
-				if (Source.IsValid() && Source->IsSourceStillValid())
+				if (TSharedPtr<ILiveLinkSourceARKit> ArkitSource = Source.Pin())
 				{
-					Source->PublishBlendShapes(SubjectName, FrameTime, BlendShapes, DeviceId);
+					if (ArkitSource->IsSourceStillValid())
+					{
+						ArkitSource->PublishBlendShapes(SubjectName, FrameTime, BlendShapes, DeviceId);
+					}
 				}
 			}
 			else
diff --git a/Engine/Plugins/Runtime/AR/AppleAR/AppleARKitFaceSupport/Source/AppleARKitFaceSupport/Private/AppleARKitLiveLinkSource.h b/Engine/Plugins/Runtime/AR/AppleAR/AppleARKitFaceSupport/Source/AppleARKitFaceSupport/Private/AppleARKitLiveLinkSource.h
--- a/Engine/Plugins/Runtime/AR/AppleAR/AppleARKitFaceSupport/Source/AppleARKitFaceSupport/Private/AppleARKitLiveLinkSource.h
+++ b/Engine/Plugins/Runtime/AR/AppleAR/AppleARKitFaceSupport/Source/AppleARKitFaceSupport/Private/AppleARKitLiveLinkSource.h
@@ -3,6 +3,7 @@
 #pragma once
 
 #include "CoreMinimal.h"
+#include "AppleARKitLiveLinkConnectionSettings.h"
 #include "AppleARKitLiveLinkSourceFactory.h"
 #include "ILiveLinkClient.h"
 #include "LiveLinkTypes.h"
@@ -725,7 +726,7 @@ class FAppleARKitLiveLinkRemoteListener :
 	public FTickableGameObject
 {
 public:
-	FAppleARKitLiveLinkRemoteListener();
+	FAppleARKitLiveLinkRemoteListener(const TSharedPtr<ILiveLinkSourceARKit>& ArkitSource);
 	virtual ~FAppleARKitLiveLinkRemoteListener();
 
 	// FTickableGameObject interface
@@ -739,10 +740,9 @@ public:
 	virtual bool IsTickableWhenPaused() const override { return true; }
 	// End FTickableGameObject
 
-	bool InitReceiveSocket();
+	bool InitReceiveSocket(int32 Port);
 
 private:
-	void InitLiveLinkSource();
 
 	/** Socket that is read from to publish remote data to this local LiveLinkClient */
 	FSocket* RecvSocket;
@@ -752,7 +752,7 @@ private:
 	FARBlendShapeMap BlendShapes;
 
 	/** The source used to publish events from the remote socket */
-	TSharedPtr<ILiveLinkSourceARKit> Source;
+	TWeakPtr<ILiveLinkSourceARKit> Source;
 };
 
 /** Forwards face blend shapes to another machine for LiveLink publication */
@@ -870,15 +870,20 @@ private:
 
 /** Publishes face blend shapes to LiveLink for use locally */
 class FAppleARKitLiveLinkSource :
-	public ILiveLinkSourceARKit
+	public ILiveLinkSourceARKit, public TSharedFromThis<FAppleARKitLiveLinkSource>
 {
 public:
 	FAppleARKitLiveLinkSource();
+	// Constructor invoked from the AppleARKitLiveLinkSourceFactory wit the settings specified by the user.
+	FAppleARKitLiveLinkSource(FAppleARKitLiveLinkConnectionSettings InConnectionSettings);
 	virtual ~FAppleARKitLiveLinkSource() {}
 
+	// Bind the remote listener to the port provided by this source's connection settings.
+	void InitializeRemoteListener();
 private:
 	// ILiveLinkSource interface
 	virtual void ReceiveClient(ILiveLinkClient* InClient, FGuid InSourceGuid) override;
+	virtual void InitializeSettings(ULiveLinkSourceSettings* Settings) override;
 	virtual bool IsSourceStillValid() const override;
 	virtual bool RequestSourceShutdown() override;
 	virtual FText GetSourceMachineName() const override;
@@ -913,4 +918,10 @@ private:
 
 	/** Used to track names changes for a given device and keep track of property names contained in subject */
 	TMap<FName, FBlendShapeStaticData> BlendShapePerDeviceMap;
+
+	/** Used to listen on a socket for incoming face data and publishing this data to livelink. */
+	TUniquePtr<FAppleARKitLiveLinkRemoteListener> RemoteListener;
+
+	/** Connection settings for this source. */
+	FAppleARKitLiveLinkConnectionSettings ConnectionSettings;
 };
diff --git a/Engine/Plugins/Runtime/AR/AppleAR/AppleARKitFaceSupport/Source/AppleARKitFaceSupport/Private/AppleARKitLiveLinkSourceFactory.cpp b/Engine/Plugins/Runtime/AR/AppleAR/AppleARKitFaceSupport/Source/AppleARKitFaceSupport/Private/AppleARKitLiveLinkSourceFactory.cpp
new file mode 100644
--- /dev/null
+++ b/Engine/Plugins/Runtime/AR/AppleAR/AppleARKitFaceSupport/Source/AppleARKitFaceSupport/Private/AppleARKitLiveLinkSourceFactory.cpp
@@ -0,0 +1,54 @@
+// Copyright Epic Games, Inc. All Rights Reserved.
+
+#include "AppleARKitLiveLinkSourceFactory.h"
+
+#include "AppleARKitLiveLinkConnectionSettings.h"
+#include "AppleARKitLiveLinkSource.h"
+#include "Features/IModularFeatures.h"
+#include "ILiveLinkClient.h"
+#include "SAppleARKitLiveLinkSourceFactory.h"
+
+#define LOCTEXT_NAMESPACE "AppleARKitLiveLinkSourceFactory"
+
+FText UAppleARKitLiveLinkSourceFactory::GetSourceDisplayName() const
+{
+	return LOCTEXT("SourceDisplayName", "Apple ARKit Source");
+}
+
+FText UAppleARKitLiveLinkSourceFactory::GetSourceTooltip() const
+{
+	return LOCTEXT("SourceTooltip", "Allows creation of an AppleARKit LiveLink source. ");
+}
+
+TSharedPtr<SWidget> UAppleARKitLiveLinkSourceFactory::BuildCreationPanel(FOnLiveLinkSourceCreated InOnLiveLinkSourceCreated) const
+{
+	return SNew(SAppleARKitLiveLinkSourceFactory)
+		.OnConnectionSettingsAccepted(FOnConnectionSettingsAccepted::CreateUObject(this, &UAppleARKitLiveLinkSourceFactory::CreateSourceFromSettings, InOnLiveLinkSourceCreated));
+}
+
+TSharedPtr<ILiveLinkSource> UAppleARKitLiveLinkSourceFactory::CreateSource(const FString& ConnectionString) const
+{
+	FAppleARKitLiveLinkConnectionSettings ConnectionSettings;
+	if (!ConnectionString.IsEmpty())
+	{
+		FAppleARKitLiveLinkConnectionSettings::StaticStruct()->ImportText(*ConnectionString, &ConnectionSettings, nullptr, PPF_None, GLog, TEXT("UAppleARKitLiveLinkSourceFactory"));
+	}
+	
+	TSharedPtr<FAppleARKitLiveLinkSource> Source = MakeShared<FAppleARKitLiveLinkSource>(ConnectionSettings);
+	Source->InitializeRemoteListener();
+
+	return Source;
+}
+
+void UAppleARKitLiveLinkSourceFactory::CreateSourceFromSettings(const FAppleARKitLiveLinkConnectionSettings& InConnectionSettings, FOnLiveLinkSourceCreated OnSourceCreated) const
+{
+	FString ConnectionString;
+	FAppleARKitLiveLinkConnectionSettings::StaticStruct()->ExportText(ConnectionString, &InConnectionSettings, nullptr, nullptr, PPF_None, nullptr);
+
+	TSharedPtr<FAppleARKitLiveLinkSource> SharedPtr = MakeShared<FAppleARKitLiveLinkSource>(InConnectionSettings);
+	SharedPtr->InitializeRemoteListener();
+
+	OnSourceCreated.ExecuteIfBound(SharedPtr, MoveTemp(ConnectionString));
+}
+
+#undef LOCTEXT_NAMESPACE
diff --git a/Engine/Plugins/Runtime/AR/AppleAR/AppleARKitFaceSupport/Source/AppleARKitFaceSupport/Private/SAppleARKitLiveLinkSourceFactory.cpp b/Engine/Plugins/Runtime/AR/AppleAR/AppleARKitFaceSupport/Source/AppleARKitFaceSupport/Private/SAppleARKitLiveLinkSourceFactory.cpp
new file mode 100644
--- /dev/null
+++ b/Engine/Plugins/Runtime/AR/AppleAR/AppleARKitFaceSupport/Source/AppleARKitFaceSupport/Private/SAppleARKitLiveLinkSourceFactory.cpp
@@ -0,0 +1,62 @@
+// Copyright Epic Games, Inc. All Rights Reserved.
+
+#include "SAppleARKitLiveLinkSourceFactory.h"
+
+#include "Modules/ModuleManager.h"
+#include "Widgets/Input/SButton.h"
+#include "Widgets/SBoxPanel.h"
+
+#if WITH_EDITOR
+#include "DetailLayoutBuilder.h"
+#include "PropertyEditorModule.h"
+#endif //WITH_EDITOR
+
+#define LOCTEXT_NAMESPACE "SAppleARKitLiveLinkSourceFactory"
+
+void SAppleARKitLiveLinkSourceFactory::Construct(const FArguments& Args)
+{
+#if WITH_EDITOR
+	OnConnectionSettingsAccepted = Args._OnConnectionSettingsAccepted;
+
+	FStructureDetailsViewArgs StructureViewArgs;
+	FDetailsViewArgs DetailArgs;
+	DetailArgs.bAllowSearch = false;
+	DetailArgs.bShowScrollBar = false;
+
+	FPropertyEditorModule& PropertyEditor = FModuleManager::Get().LoadModuleChecked<FPropertyEditorModule>(TEXT("PropertyEditor"));
+
+	StructOnScope = MakeShared<FStructOnScope>(FAppleARKitLiveLinkConnectionSettings::StaticStruct());
+	CastChecked<UScriptStruct>(StructOnScope->GetStruct())->CopyScriptStruct(StructOnScope->GetStructMemory(), &ConnectionSettings);
+	StructureDetailsView = PropertyEditor.CreateStructureDetailView(DetailArgs, StructureViewArgs, StructOnScope);
+
+	ChildSlot
+		[
+			SNew(SVerticalBox)
+				+ SVerticalBox::Slot()
+				.FillHeight(1.f)
+				[
+					StructureDetailsView->GetWidget().ToSharedRef()
+				]
+				+ SVerticalBox::Slot()
+				.HAlign(HAlign_Right)
+				.AutoHeight()
+				[
+					SNew(SButton)
+						.OnClicked(this, &SAppleARKitLiveLinkSourceFactory::OnSettingsAccepted)
+						.Text(LOCTEXT("AddSource", "Add"))
+				]
+		];
+#endif //WITH_EDITOR
+}
+
+FReply SAppleARKitLiveLinkSourceFactory::OnSettingsAccepted()
+{
+#if WITH_EDITOR
+	CastChecked<UScriptStruct>(StructOnScope->GetStruct())->CopyScriptStruct(&ConnectionSettings, StructOnScope->GetStructMemory());
+	OnConnectionSettingsAccepted.ExecuteIfBound(ConnectionSettings);
+#endif //WITH_EDITOR
+
+	return FReply::Handled();
+}
+
+#undef LOCTEXT_NAMESPACE
\ No newline at end of file
diff --git a/Engine/Plugins/Runtime/AR/AppleAR/AppleARKitFaceSupport/Source/AppleARKitFaceSupport/Private/SAppleARKitLiveLinkSourceFactory.h b/Engine/Plugins/Runtime/AR/AppleAR/AppleARKitFaceSupport/Source/AppleARKitFaceSupport/Private/SAppleARKitLiveLinkSourceFactory.h
new file mode 100644
--- /dev/null
+++ b/Engine/Plugins/Runtime/AR/AppleAR/AppleARKitFaceSupport/Source/AppleARKitFaceSupport/Private/SAppleARKitLiveLinkSourceFactory.h
@@ -0,0 +1,39 @@
+// Copyright Epic Games, Inc. All Rights Reserved.
+
+#pragma once
+
+#include "AppleARKitLiveLinkConnectionSettings.h"
+#include "Widgets/SCompoundWidget.h"
+#include "Widgets/DeclarativeSyntaxSupport.h"
+
+#if WITH_EDITOR
+#include "IStructureDetailsView.h"
+#endif //WITH_EDITOR
+
+#include "Input/Reply.h"
+
+struct FAppleARKitLiveLinkConnectionSettings;
+
+DECLARE_DELEGATE_OneParam(FOnConnectionSettingsAccepted, const FAppleARKitLiveLinkConnectionSettings&);
+
+class SAppleARKitLiveLinkSourceFactory : public SCompoundWidget
+{
+	SLATE_BEGIN_ARGS(SAppleARKitLiveLinkSourceFactory)
+		{}
+		SLATE_EVENT(FOnConnectionSettingsAccepted, OnConnectionSettingsAccepted)
+	SLATE_END_ARGS()
+
+	void Construct(const FArguments& Args);
+
+
+private:
+	FAppleARKitLiveLinkConnectionSettings ConnectionSettings;
+
+#if WITH_EDITOR
+	TSharedPtr<FStructOnScope> StructOnScope;
+	TSharedPtr<IStructureDetailsView> StructureDetailsView;
+#endif //WITH_EDITOR
+
+	FReply OnSettingsAccepted();
+	FOnConnectionSettingsAccepted OnConnectionSettingsAccepted;
+};
\ No newline at end of file
diff --git a/Engine/Plugins/Runtime/AR/AppleAR/AppleARKitFaceSupport/Source/AppleARKitFaceSupport/Public/AppleARKitLiveLinkSourceFactory.h b/Engine/Plugins/Runtime/AR/AppleAR/AppleARKitFaceSupport/Source/AppleARKitFaceSupport/Public/AppleARKitLiveLinkSourceFactory.h
--- a/Engine/Plugins/Runtime/AR/AppleAR/AppleARKitFaceSupport/Source/AppleARKitFaceSupport/Public/AppleARKitLiveLinkSourceFactory.h
+++ b/Engine/Plugins/Runtime/AR/AppleAR/AppleARKitFaceSupport/Source/AppleARKitFaceSupport/Public/AppleARKitLiveLinkSourceFactory.h
@@ -6,6 +6,9 @@
 #include "ARTrackable.h"
 #include "ILiveLinkSource.h"
 
+#include "AppleARKitLiveLinkSourceFactory.generated.h"
+
+struct FAppleARKitLiveLinkConnectionSettings;
 class UTimecodeProvider;
 
 /** Interface that publishes face ar blend shape information */
@@ -30,8 +33,8 @@ public:
 	/** Creates a face mesh source that will autobind to the tracked face mesh */
 	static TSharedPtr<ILiveLinkSourceARKit> CreateLiveLinkSource();
 
-	/** Creates the singleton listener that will receive remote events and publish them locally */
-	static void CreateLiveLinkRemoteListener();
+	/** Creates a AppleARKit source that holds a livelink remote listener. */
+	static TSharedPtr<ILiveLinkSourceARKit> CreateLiveLinkSource(const FAppleARKitLiveLinkConnectionSettings& ConnectionSettings);
 
 	/** Creates the publisher that will send remote events to a specified IP */
 	static TSharedPtr<IARKitBlendShapePublisher, ESPMode::ThreadSafe> CreateLiveLinkRemotePublisher(const FString& RemoteAddr = FString());
@@ -39,3 +42,20 @@ public:
 	/** Creates the publisher that will write the curve data to disk */
 	static TSharedPtr<IARKitBlendShapePublisher, ESPMode::ThreadSafe> CreateLiveLinkLocalFileWriter();
 };
+
+UCLASS()
+class APPLEARKITFACESUPPORT_API UAppleARKitLiveLinkSourceFactory : public ULiveLinkSourceFactory
+{
+public:
+	GENERATED_BODY()
+
+	virtual FText GetSourceDisplayName() const override;
+	virtual FText GetSourceTooltip() const override;
+
+	virtual EMenuType GetMenuType() const override { return EMenuType::SubPanel; }
+	virtual TSharedPtr<SWidget> BuildCreationPanel(FOnLiveLinkSourceCreated OnLiveLinkSourceCreated) const override;
+	virtual TSharedPtr<ILiveLinkSource> CreateSource(const FString& ConnectionString) const override;
+
+private:
+	void CreateSourceFromSettings(const FAppleARKitLiveLinkConnectionSettings& ConnectionSettings, FOnLiveLinkSourceCreated OnSourceCreated) const;
+};
\ No newline at end of file
diff --git a/Engine/Programs/LiveLinkHub/Config/DefaultEngine.ini b/Engine/Programs/LiveLinkHub/Config/DefaultEngine.ini
--- a/Engine/Programs/LiveLinkHub/Config/DefaultEngine.ini
+++ b/Engine/Programs/LiveLinkHub/Config/DefaultEngine.ini
@@ -19,6 +19,10 @@ GameNameShort=LLH
 TemplateMacPlist=(FilePath="/Engine/Source/Programs/LiveLinkHub/Resources/Mac/Info.plist")
 
 
+[/Script/AppleARKit.AppleARKitSettings]
+bEnableLiveLinkForFaceTracking=true
+
+
 [/Script/UnrealEd.ProjectPackagingSettings]
 bShareMaterialShaderCode=false
 
