Fix cplusplus.PureVirtualCall

#rnx

diff --git a/Engine/Plugins/Runtime/AR/AppleAR/AppleARKitFaceSupport/Source/AppleARKitFaceSupport/Private/AppleARKitFaceSupportImpl.cpp b/Engine/Plugins/Runtime/AR/AppleAR/AppleARKitFaceSupport/Source/AppleARKitFaceSupport/Private/AppleARKitFaceSupportImpl.cpp
--- a/Engine/Plugins/Runtime/AR/AppleAR/AppleARKitFaceSupport/Source/AppleARKitFaceSupport/Private/AppleARKitFaceSupportImpl.cpp
+++ b/Engine/Plugins/Runtime/AR/AppleAR/AppleARKitFaceSupport/Source/AppleARKitFaceSupport/Private/AppleARKitFaceSupportImpl.cpp
@@ -97,7 +97,16 @@ void FAppleARKitFaceSupport::Init()
 
 void FAppleARKitFaceSupport::Shutdown()
 {
+	if (RemoteLiveLinkPublisher.IsValid())
+	{
+		RemoteLiveLinkPublisher->Shutdown();
+	}
 	RemoteLiveLinkPublisher = nullptr;
+
+	if (LiveLinkFileWriter.IsValid())
+	{
+		LiveLinkFileWriter->Shutdown();
+	}
 	LiveLinkFileWriter = nullptr;
 
 	IModularFeatures::Get().UnregisterModularFeature(GetModularFeatureName(), this);
diff --git a/Engine/Plugins/Runtime/AR/AppleAR/AppleARKitFaceSupport/Source/AppleARKitFaceSupport/Private/AppleARKitLiveLinkSource.cpp b/Engine/Plugins/Runtime/AR/AppleAR/AppleARKitFaceSupport/Source/AppleARKitFaceSupport/Private/AppleARKitLiveLinkSource.cpp
--- a/Engine/Plugins/Runtime/AR/AppleAR/AppleARKitFaceSupport/Source/AppleARKitFaceSupport/Private/AppleARKitLiveLinkSource.cpp
+++ b/Engine/Plugins/Runtime/AR/AppleAR/AppleARKitFaceSupport/Source/AppleARKitFaceSupport/Private/AppleARKitLiveLinkSource.cpp
@@ -507,7 +507,7 @@ FAppleARKitLiveLinkFileWriter::FAppleARKitLiveLinkFileWriter(const TCHAR* InFile
 	UAppleARKitSettings::CreateFaceTrackingLogDir();
 }
 
-FAppleARKitLiveLinkFileWriter::~FAppleARKitLiveLinkFileWriter()
+void FAppleARKitLiveLinkFileWriter::Shutdown()
 {
 	// Save on close if desired
 	if (!GetMutableDefault<UAppleARKitSettings>()->ShouldFaceTrackingLogPerFrame())
diff --git a/Engine/Plugins/Runtime/AR/AppleAR/AppleARKitFaceSupport/Source/AppleARKitFaceSupport/Private/AppleARKitLiveLinkSource.h b/Engine/Plugins/Runtime/AR/AppleAR/AppleARKitFaceSupport/Source/AppleARKitFaceSupport/Private/AppleARKitLiveLinkSource.h
--- a/Engine/Plugins/Runtime/AR/AppleAR/AppleARKitFaceSupport/Source/AppleARKitFaceSupport/Private/AppleARKitLiveLinkSource.h
+++ b/Engine/Plugins/Runtime/AR/AppleAR/AppleARKitFaceSupport/Source/AppleARKitFaceSupport/Private/AppleARKitLiveLinkSource.h
@@ -787,8 +787,9 @@ class FAppleARKitLiveLinkFileWriter :
 	public IARKitBlendShapePublisher
 {
 public:
-	virtual ~FAppleARKitLiveLinkFileWriter();
+	virtual ~FAppleARKitLiveLinkFileWriter() {}
 	virtual void SetTimecodeProvider(UTimecodeProvider* InTimecodeProvider = nullptr) override;
+	virtual void Shutdown() override;
 
 protected:
 	FAppleARKitLiveLinkFileWriter(const TCHAR* InFileExtension);
@@ -880,6 +881,8 @@ public:
 
 	// Bind the remote listener to the port provided by this source's connection settings.
 	void InitializeRemoteListener();
+
+	virtual void Shutdown() override {}
 private:
 	// ILiveLinkSource interface
 	virtual void ReceiveClient(ILiveLinkClient* InClient, FGuid InSourceGuid) override;
diff --git a/Engine/Plugins/Runtime/AR/AppleAR/AppleARKitFaceSupport/Source/AppleARKitFaceSupport/Public/AppleARKitLiveLinkSourceFactory.h b/Engine/Plugins/Runtime/AR/AppleAR/AppleARKitFaceSupport/Source/AppleARKitFaceSupport/Public/AppleARKitLiveLinkSourceFactory.h
--- a/Engine/Plugins/Runtime/AR/AppleAR/AppleARKitFaceSupport/Source/AppleARKitFaceSupport/Public/AppleARKitLiveLinkSourceFactory.h
+++ b/Engine/Plugins/Runtime/AR/AppleAR/AppleARKitFaceSupport/Source/AppleARKitFaceSupport/Public/AppleARKitLiveLinkSourceFactory.h
@@ -17,6 +17,7 @@ class APPLEARKITFACESUPPORT_API IARKitBlendShapePublisher
 public:
 	virtual void SetTimecodeProvider(UTimecodeProvider* InTimecodeProvider = nullptr) {}
 	virtual void PublishBlendShapes(FName SubjectName, const FQualifiedFrameTime& FrameTime, const FARBlendShapeMap& BlendShapes, FName DeviceID = NAME_None) = 0;
+	virtual void Shutdown() = 0;
 };
 
 /** Interface that publishes face ar blend shape information via LiveLink */
