Fix ios compile error

#rnx

diff --git a/Engine/Plugins/Runtime/AR/AppleAR/AppleARKitFaceSupport/Source/AppleARKitFaceSupport/Private/AppleARKitLiveLinkSource.h b/Engine/Plugins/Runtime/AR/AppleAR/AppleARKitFaceSupport/Source/AppleARKitFaceSupport/Private/AppleARKitLiveLinkSource.h
--- a/Engine/Plugins/Runtime/AR/AppleAR/AppleARKitFaceSupport/Source/AppleARKitFaceSupport/Private/AppleARKitLiveLinkSource.h
+++ b/Engine/Plugins/Runtime/AR/AppleAR/AppleARKitFaceSupport/Source/AppleARKitFaceSupport/Private/AppleARKitLiveLinkSource.h
@@ -766,6 +766,7 @@ public:
 
 	// IARKitBlendShapePublisher interface
 	virtual void PublishBlendShapes(FName SubjectName, const FQualifiedFrameTime& FrameTime, const FARBlendShapeMap& FaceBlendShapes, FName DeviceId = NAME_None) override;
+	virtual void Shutdown() override {}
 	// End IARKitBlendShapePublisher
 
 	bool InitSendSocket();
