Set a min / max value for ARKit connection settings

diff --git a/Engine/Plugins/Runtime/AR/AppleAR/AppleARKitFaceSupport/Source/AppleARKitFaceSupport/Private/AppleARKitLiveLinkConnectionSettings.h b/Engine/Plugins/Runtime/AR/AppleAR/AppleARKitFaceSupport/Source/AppleARKitFaceSupport/Private/AppleARKitLiveLinkConnectionSettings.h
--- a/Engine/Plugins/Runtime/AR/AppleAR/AppleARKitFaceSupport/Source/AppleARKitFaceSupport/Private/AppleARKitLiveLinkConnectionSettings.h
+++ b/Engine/Plugins/Runtime/AR/AppleAR/AppleARKitFaceSupport/Source/AppleARKitFaceSupport/Private/AppleARKitLiveLinkConnectionSettings.h
@@ -10,6 +10,6 @@ struct APPLEARKITFACESUPPORT_API FAppleARKitLiveLinkConnectionSettings
 	GENERATED_BODY()
 
 	/** The port to use when listening/sending LiveLink face blend shapes via the network */
-	UPROPERTY(EditAnywhere, Category = "Connection Settings")
+	UPROPERTY(EditAnywhere, Category = "Connection Settings", meta = (ClampMin = "1", ClampMax = "32765"))
 	int32 Port = 11111;
 };
